Optimus: Move skeletal modeling tools into their own plugin

diff --git a/Engine/Plugins/Animation/DeformerGraph/DeformerGraph.uplugin b/Engine/Plugins/Animation/DeformerGraph/DeformerGraph.uplugin
--- a/Engine/Plugins/Animation/DeformerGraph/DeformerGraph.uplugin
+++ b/Engine/Plugins/Animation/DeformerGraph/DeformerGraph.uplugin
@@ -34,26 +34,12 @@
 			"Name": "OptimusEditor",
 			"Type": "Editor",
 			"LoadingPhase": "PostEngineInit"
-		},
-		{
-			"Name": "OptimusTools",
-			"Type": "Editor",
-			"LoadingPhase": "Default",
-			"PlatformAllowList": [ "Win64" ]
 		}
 	],
 	"Plugins": [
 		{
 			"Name": "ComputeFramework",
 			"Enabled": true
-		},
-		{
-			"Name": "MeshModelingToolsetExp",
-			"Enabled": true
-		},
-		{
-			"Name": "ModelingToolsEditorMode",
-			"Enabled": true
 		}
 	]
 }
diff --git a/Engine/Plugins/Animation/DeformerGraph/Source/OptimusTools/OptimusTools.Build.cs b/Engine/Plugins/Animation/DeformerGraph/Source/OptimusTools/OptimusTools.Build.cs
deleted file mode 100644
diff --git a/Engine/Plugins/Animation/DeformerGraph/Source/OptimusTools/Private/OptimusToolsCommands.cpp b/Engine/Plugins/Animation/DeformerGraph/Source/OptimusTools/Private/OptimusToolsCommands.cpp
deleted file mode 100644
diff --git a/Engine/Plugins/Animation/DeformerGraph/Source/OptimusTools/Private/OptimusToolsCommands.h b/Engine/Plugins/Animation/DeformerGraph/Source/OptimusTools/Private/OptimusToolsCommands.h
deleted file mode 100644
diff --git a/Engine/Plugins/Animation/DeformerGraph/Source/OptimusTools/Private/OptimusToolsModule.cpp b/Engine/Plugins/Animation/DeformerGraph/Source/OptimusTools/Private/OptimusToolsModule.cpp
deleted file mode 100644
diff --git a/Engine/Plugins/Animation/DeformerGraph/Source/OptimusTools/Private/OptimusToolsModule.h b/Engine/Plugins/Animation/DeformerGraph/Source/OptimusTools/Private/OptimusToolsModule.h
deleted file mode 100644
diff --git a/Engine/Plugins/Animation/DeformerGraph/Source/OptimusTools/Private/OptimusToolsSkelMeshEditorMode.cpp b/Engine/Plugins/Animation/DeformerGraph/Source/OptimusTools/Private/OptimusToolsSkelMeshEditorMode.cpp
deleted file mode 100644
diff --git a/Engine/Plugins/Animation/DeformerGraph/Source/OptimusTools/Private/OptimusToolsSkelMeshEditorMode.h b/Engine/Plugins/Animation/DeformerGraph/Source/OptimusTools/Private/OptimusToolsSkelMeshEditorMode.h
deleted file mode 100644
diff --git a/Engine/Plugins/Animation/DeformerGraph/Source/OptimusTools/Private/OptimusToolsSkelMeshEditorModeToolkit.cpp b/Engine/Plugins/Animation/DeformerGraph/Source/OptimusTools/Private/OptimusToolsSkelMeshEditorModeToolkit.cpp
deleted file mode 100644
diff --git a/Engine/Plugins/Animation/DeformerGraph/Source/OptimusTools/Private/OptimusToolsSkelMeshEditorModeToolkit.h b/Engine/Plugins/Animation/DeformerGraph/Source/OptimusTools/Private/OptimusToolsSkelMeshEditorModeToolkit.h
deleted file mode 100644
diff --git a/Engine/Plugins/Animation/DeformerGraph/Source/OptimusTools/Private/OptimusToolsStyle.cpp b/Engine/Plugins/Animation/DeformerGraph/Source/OptimusTools/Private/OptimusToolsStyle.cpp
deleted file mode 100644
diff --git a/Engine/Plugins/Animation/DeformerGraph/Source/OptimusTools/Private/OptimusToolsStyle.h b/Engine/Plugins/Animation/DeformerGraph/Source/OptimusTools/Private/OptimusToolsStyle.h
deleted file mode 100644
diff --git a/Engine/Plugins/Experimental/Animation/SkeletalMeshModelingTools/SkeletalMeshModelingTools.uplugin b/Engine/Plugins/Experimental/Animation/SkeletalMeshModelingTools/SkeletalMeshModelingTools.uplugin
new file mode 100644
--- /dev/null
+++ b/Engine/Plugins/Experimental/Animation/SkeletalMeshModelingTools/SkeletalMeshModelingTools.uplugin
@@ -0,0 +1,38 @@
+﻿{
+	"FileVersion": 3,
+	"Version": 1,
+	"VersionName": "0.1",
+	"FriendlyName": "SkeletalMeshModelingTools",
+	"Description": "Modeling Tools for the Skeletal Mesh Editor",
+	"Category" : "Animation",
+	"CreatedBy" : "Epic Games, Inc.",
+	"CreatedByURL": "http://epicgames.com/",
+	"DocsURL": "",
+	"MarketplaceURL": "",
+	"EnabledByDefault" : false,
+	"CanContainContent": true,
+	"IsExperimentalVersion": true,
+	"Installed": false,
+	"Modules": [
+		{
+			"Name": "SkeletalMeshModelingTools",
+			"Type": "Editor",
+			"LoadingPhase": "Default",
+			"PlatformAllowList": [ "Win64" ]
+		}
+	],
+	"Plugins": [
+		{
+			"Name": "ComputeFramework",
+			"Enabled": true
+		},
+		{
+			"Name": "MeshModelingToolsetExp",
+			"Enabled": true
+		},
+		{
+			"Name": "ModelingToolsEditorMode",
+			"Enabled": true
+		}
+	]
+}
diff --git a/Engine/Plugins/Experimental/Animation/SkeletalMeshModelingTools/Source/SkeletalMeshModelingTools/Private/SkeletalMeshModelingToolsCommands.cpp b/Engine/Plugins/Experimental/Animation/SkeletalMeshModelingTools/Source/SkeletalMeshModelingTools/Private/SkeletalMeshModelingToolsCommands.cpp
new file mode 100644
--- /dev/null
+++ b/Engine/Plugins/Experimental/Animation/SkeletalMeshModelingTools/Source/SkeletalMeshModelingTools/Private/SkeletalMeshModelingToolsCommands.cpp
@@ -0,0 +1,18 @@
+// Copyright Epic Games, Inc. All Rights Reserved.
+
+#include "SkeletalMeshModelingToolsCommands.h"
+
+#define LOCTEXT_NAMESPACE "SkeletalMeshModelingToolsCommands"
+
+
+void FSkeletalMeshModelingToolsCommands::RegisterCommands()
+{
+	UI_COMMAND(ToggleModelingToolsMode, "Enable Modeling Tools", "Toggles modeling tools on or off.", EUserInterfaceActionType::ToggleButton, FInputChord());
+}
+
+const FSkeletalMeshModelingToolsCommands& FSkeletalMeshModelingToolsCommands::Get()
+{
+	return TCommands<FSkeletalMeshModelingToolsCommands>::Get();
+}
+
+#undef LOCTEXT_NAMESPACE
diff --git a/Engine/Plugins/Experimental/Animation/SkeletalMeshModelingTools/Source/SkeletalMeshModelingTools/Private/SkeletalMeshModelingToolsCommands.h b/Engine/Plugins/Experimental/Animation/SkeletalMeshModelingTools/Source/SkeletalMeshModelingTools/Private/SkeletalMeshModelingToolsCommands.h
new file mode 100644
--- /dev/null
+++ b/Engine/Plugins/Experimental/Animation/SkeletalMeshModelingTools/Source/SkeletalMeshModelingTools/Private/SkeletalMeshModelingToolsCommands.h
@@ -0,0 +1,27 @@
+// Copyright Epic Games, Inc. All Rights Reserved.
+
+#pragma once
+
+#include "CoreMinimal.h"
+#include "Styling/AppStyle.h"
+#include "Framework/Commands/Commands.h"
+
+
+class FSkeletalMeshModelingToolsCommands : public TCommands<FSkeletalMeshModelingToolsCommands>
+{
+public:
+	FSkeletalMeshModelingToolsCommands()
+	    : TCommands<FSkeletalMeshModelingToolsCommands>(
+	    	TEXT("SkeletalMeshModelingTools"),
+	    	NSLOCTEXT("Contexts", "SkeletalMeshModelingToolsCommands", "Skeletal Mesh Modeling Tools"),
+	    	NAME_None,
+	    	FAppStyle::GetAppStyleSetName())
+	{
+	}
+
+	void RegisterCommands() override;
+	static const FSkeletalMeshModelingToolsCommands& Get();
+
+	// Modeling tools commands
+	TSharedPtr<FUICommandInfo> ToggleModelingToolsMode;
+};
diff --git a/Engine/Plugins/Experimental/Animation/SkeletalMeshModelingTools/Source/SkeletalMeshModelingTools/Private/SkeletalMeshModelingToolsEditorMode.cpp b/Engine/Plugins/Experimental/Animation/SkeletalMeshModelingTools/Source/SkeletalMeshModelingTools/Private/SkeletalMeshModelingToolsEditorMode.cpp
new file mode 100644
--- /dev/null
+++ b/Engine/Plugins/Experimental/Animation/SkeletalMeshModelingTools/Source/SkeletalMeshModelingTools/Private/SkeletalMeshModelingToolsEditorMode.cpp
@@ -0,0 +1,248 @@
+// Copyright Epic Games, Inc. All Rights Reserved.
+
+#include "SkeletalMeshModelingToolsEditorMode.h"
+
+#include "SkeletalMeshModelingToolsEditorModeToolkit.h"
+#include "SkeletalMeshModelingToolsCommands.h"
+
+#include "BaseGizmos/TransformGizmoUtil.h"
+#include "EdMode.h"
+#include "EdModeInteractiveToolsContext.h"
+#include "EditorModeManager.h"
+#include "EditorViewportClient.h"
+#include "IStylusInputModule.h"
+#include "IStylusState.h"
+#include "ToolTargets/SkeletalMeshComponentToolTarget.h"
+
+#include "DeformMeshPolygonsTool.h"
+#include "DisplaceMeshTool.h"
+#include "DynamicMeshSculptTool.h"
+#include "EditMeshPolygonsTool.h"
+#include "HoleFillTool.h"
+#include "LatticeDeformerTool.h"
+#include "MeshAttributePaintTool.h"
+#include "MeshSpaceDeformerTool.h"
+#include "MeshVertexSculptTool.h"
+#include "ModelingToolsManagerActions.h"
+#include "OffsetMeshTool.h"
+#include "PolygonOnMeshTool.h"
+#include "ProjectToTargetTool.h"
+#include "RemeshMeshTool.h"
+#include "RemoveOccludedTrianglesTool.h"
+#include "SimplifyMeshTool.h"
+#include "SkinWeightsBindingTool.h"
+#include "SkinWeightsPaintTool.h"
+#include "SmoothMeshTool.h"
+#include "ToolTargetManager.h"
+#include "WeldMeshEdgesTool.h"
+
+
+#define LOCTEXT_NAMESPACE "SkeletalMeshModelingToolsEditorMode"
+
+
+// FStylusStateTracker registers itself as a listener for stylus events and implements
+// the IToolStylusStateProviderAPI interface, which allows MeshSurfacePointTool implementations
+ // to query for the pen pressure.
+//
+// This is kind of a hack. Unfortunately the current Stylus module is a Plugin so it
+// cannot be used in the base ToolsFramework, and we need this in the Mode as a workaround.
+//
+class FStylusStateTracker : public IStylusMessageHandler, public IToolStylusStateProviderAPI
+{
+public:
+	const IStylusInputDevice* ActiveDevice = nullptr;
+	int32 ActiveDeviceIndex = -1;
+
+	bool bPenDown = false;
+	float ActivePressure = 1.0;
+
+	FStylusStateTracker()
+	{
+		UStylusInputSubsystem* StylusSubsystem = GEditor->GetEditorSubsystem<UStylusInputSubsystem>();
+		StylusSubsystem->AddMessageHandler(*this);
+
+		ActiveDevice = FindFirstPenDevice(StylusSubsystem, ActiveDeviceIndex);
+		bPenDown = false;
+	}
+
+	virtual ~FStylusStateTracker()
+	{
+		UStylusInputSubsystem* StylusSubsystem = GEditor->GetEditorSubsystem<UStylusInputSubsystem>();
+		StylusSubsystem->RemoveMessageHandler(*this);
+	}
+
+	void OnStylusStateChanged(const FStylusState& NewState, int32 StylusIndex) override
+	{
+		if (ActiveDevice == nullptr)
+		{
+			UStylusInputSubsystem* StylusSubsystem = GEditor->GetEditorSubsystem<UStylusInputSubsystem>();
+			ActiveDevice = FindFirstPenDevice(StylusSubsystem, ActiveDeviceIndex);
+			bPenDown = false;
+		}
+		if (ActiveDevice != nullptr && ActiveDeviceIndex == StylusIndex)
+		{
+			bPenDown = NewState.IsStylusDown();
+			ActivePressure = NewState.GetPressure();
+		}
+	}
+
+
+	bool HaveActiveStylusState() const
+	{
+		return ActiveDevice != nullptr && bPenDown;
+	}
+
+	static const IStylusInputDevice* FindFirstPenDevice(const UStylusInputSubsystem* StylusSubsystem, int32& ActiveDeviceOut)
+	{
+		int32 NumDevices = StylusSubsystem->NumInputDevices();
+		for (int32 k = 0; k < NumDevices; ++k)
+		{
+			const IStylusInputDevice* Device = StylusSubsystem->GetInputDevice(k);
+			const TArray<EStylusInputType>& Inputs = Device->GetSupportedInputs();
+			for (EStylusInputType Input : Inputs)
+			{
+				if (Input == EStylusInputType::Pressure)
+				{
+					ActiveDeviceOut = k;
+					return Device;
+				}
+			}
+		}
+		return nullptr;
+	}
+
+
+
+	// IToolStylusStateProviderAPI implementation
+	virtual float GetCurrentPressure() const override
+	{
+		return (ActiveDevice != nullptr && bPenDown) ? ActivePressure : 1.0f;
+	}
+
+};
+
+
+// NOTE: This is a simple proxy at the moment. In the future we want to pull in more of the 
+// modeling tools as we add support in the skelmesh storage.
+
+const FEditorModeID USkeletalMeshModelingToolsEditorMode::Id("SkeletalMeshModelingToolsEditorMode");
+
+
+USkeletalMeshModelingToolsEditorMode::USkeletalMeshModelingToolsEditorMode() 
+{
+	Info = FEditorModeInfo(Id, LOCTEXT("SkeletalMeshModelingMode", "Skeletal Mesh Modeling"), FSlateIcon(), false);
+}
+
+
+USkeletalMeshModelingToolsEditorMode::USkeletalMeshModelingToolsEditorMode(FVTableHelper& Helper) :
+	UBaseLegacyWidgetEdMode(Helper)
+{
+	
+}
+
+
+USkeletalMeshModelingToolsEditorMode::~USkeletalMeshModelingToolsEditorMode()
+{
+	// Implemented in the CPP file so that the destructor for TUniquePtr<FStylusStateTracker> gets correctly compiled.
+}
+
+
+void USkeletalMeshModelingToolsEditorMode::Initialize()
+{
+	UBaseLegacyWidgetEdMode::Initialize();
+}
+
+
+void USkeletalMeshModelingToolsEditorMode::Enter()
+{
+	UEdMode::Enter();
+
+	GetInteractiveToolsContext()->TargetManager->AddTargetFactory(NewObject<USkeletalMeshComponentToolTargetFactory>(GetInteractiveToolsContext()->TargetManager));
+
+	StylusStateTracker = MakeUnique<FStylusStateTracker>();
+	// register gizmo helper
+	UE::TransformGizmoUtil::RegisterTransformGizmoContextObject(GetInteractiveToolsContext());
+
+	const FModelingToolsManagerCommands& ToolManagerCommands = FModelingToolsManagerCommands::Get();
+
+		/*
+		ToolbarBuilder.AddToolBarButton(Commands.BeginPolyEditTool);
+		ToolbarBuilder.AddToolBarButton(Commands.BeginPolyDeformTool);
+		ToolbarBuilder.AddToolBarButton(Commands.BeginHoleFillTool);
+		ToolbarBuilder.AddToolBarButton(Commands.BeginPolygonCutTool);
+	}
+	else if (PaletteName == ProcessingTabName)
+	{
+		ToolbarBuilder.AddToolBarButton(Commands.BeginSimplifyMeshTool);
+		ToolbarBuilder.AddToolBarButton(Commands.BeginRemeshMeshTool);
+		ToolbarBuilder.AddToolBarButton(Commands.BeginWeldEdgesTool);
+		ToolbarBuilder.AddToolBarButton(Commands.BeginRemoveOccludedTrianglesTool);
+		ToolbarBuilder.AddToolBarButton(Commands.BeginProjectToTargetTool);
+	}
+	else if (PaletteName == DeformTabName)
+	{
+		ToolbarBuilder.AddToolBarButton(Commands.BeginSculptMeshTool);
+		ToolbarBuilder.AddToolBarButton(Commands.BeginRemeshSculptMeshTool);
+		ToolbarBuilder.AddToolBarButton(Commands.BeginSmoothMeshTool);
+		ToolbarBuilder.AddToolBarButton(Commands.BeginOffsetMeshTool);
+		ToolbarBuilder.AddToolBarButton(Commands.BeginMeshSpaceDeformerTool);
+		ToolbarBuilder.AddToolBarButton(Commands.BeginLatticeDeformerTool);
+		ToolbarBuilder.AddToolBarButton(Commands.BeginDisplaceMeshTool);
+		*/
+
+	RegisterTool(ToolManagerCommands.BeginPolyEditTool, TEXT("BeginPolyEditTool"), NewObject<UEditMeshPolygonsToolBuilder>());
+	RegisterTool(ToolManagerCommands.BeginPolyDeformTool, TEXT("BeginPolyDeformTool"), NewObject<UDeformMeshPolygonsToolBuilder>());
+	RegisterTool(ToolManagerCommands.BeginHoleFillTool, TEXT("BeginHoleFillTool"), NewObject<UHoleFillToolBuilder>());
+	RegisterTool(ToolManagerCommands.BeginPolygonCutTool, TEXT("BeginPolyCutTool"), NewObject<UPolygonOnMeshToolBuilder>());
+
+	RegisterTool(ToolManagerCommands.BeginSimplifyMeshTool, TEXT("BeginSimplifyMeshTool"), NewObject<USimplifyMeshToolBuilder>());
+	RegisterTool(ToolManagerCommands.BeginRemeshMeshTool, TEXT("BeginRemeshMeshTool"), NewObject<URemeshMeshToolBuilder>());
+	RegisterTool(ToolManagerCommands.BeginWeldEdgesTool, TEXT("BeginWeldEdgesTool"), NewObject<UWeldMeshEdgesToolBuilder>());
+	RegisterTool(ToolManagerCommands.BeginRemoveOccludedTrianglesTool, TEXT("BeginRemoveOccludedTrianglesTool"), NewObject<URemoveOccludedTrianglesToolBuilder>());
+	RegisterTool(ToolManagerCommands.BeginProjectToTargetTool, TEXT("BeginProjectToTargetTool"), NewObject<UProjectToTargetToolBuilder>());
+	
+
+	auto MoveVerticesToolBuilder = NewObject<UMeshVertexSculptToolBuilder>();
+	MoveVerticesToolBuilder->StylusAPI = StylusStateTracker.Get();
+	RegisterTool(ToolManagerCommands.BeginSculptMeshTool, TEXT("BeginSculptMeshTool"), MoveVerticesToolBuilder);
+
+	auto DynaSculptToolBuilder = NewObject<UDynamicMeshSculptToolBuilder>();
+	DynaSculptToolBuilder->bEnableRemeshing = true;
+	DynaSculptToolBuilder->StylusAPI = StylusStateTracker.Get();
+	RegisterTool(ToolManagerCommands.BeginRemeshSculptMeshTool, TEXT("BeginRemeshSculptMeshTool"), DynaSculptToolBuilder);
+	
+	RegisterTool(ToolManagerCommands.BeginSmoothMeshTool, TEXT("BeginSmoothMeshTool"), NewObject<USmoothMeshToolBuilder>());
+	RegisterTool(ToolManagerCommands.BeginOffsetMeshTool, TEXT("BeginOffsetMeshTool"), NewObject<UOffsetMeshToolBuilder>());
+	RegisterTool(ToolManagerCommands.BeginMeshSpaceDeformerTool, TEXT("BeginMeshSpaceDeformerTool"), NewObject<UMeshSpaceDeformerToolBuilder>());
+	RegisterTool(ToolManagerCommands.BeginLatticeDeformerTool, TEXT("BeginLatticeDeformerTool"), NewObject<ULatticeDeformerToolBuilder>());
+	RegisterTool(ToolManagerCommands.BeginDisplaceMeshTool, TEXT("BeginDisplaceMeshTool"), NewObject<UDisplaceMeshToolBuilder>());
+
+	// RegisterTool(ToolManagerCommands.BeginMeshAttributePaintTool, TEXT("BeginMeshAttributePaintTool"), NewObject<UMeshAttributePaintToolBuilder>());
+	RegisterTool(ToolManagerCommands.BeginSkinWeightsPaintTool, TEXT("BeginSkinWeightsPaintTool"), NewObject<USkinWeightsPaintToolBuilder>());
+	RegisterTool(ToolManagerCommands.BeginSkinWeightsBindingTool, TEXT("BeginSkinWeightsBindingTool"), NewObject<USkinWeightsBindingToolBuilder>());
+
+	GetInteractiveToolsContext()->ToolManager->SelectActiveToolType(EToolSide::Left, TEXT("BeginSkinWeightsPaintTool"));
+}
+
+
+void USkeletalMeshModelingToolsEditorMode::Exit()
+{
+	StylusStateTracker = nullptr;
+
+	UEdMode::Exit();
+}
+
+
+void USkeletalMeshModelingToolsEditorMode::CreateToolkit()
+{
+	Toolkit = MakeShareable(new FSkeletalMeshModelingToolsEditorModeToolkit);
+}
+
+
+void USkeletalMeshModelingToolsEditorMode::Tick(FEditorViewportClient* ViewportClient, float DeltaTime)
+{
+	Super::Tick(ViewportClient, DeltaTime);
+}
+
+
+#undef LOCTEXT_NAMESPACE
diff --git a/Engine/Plugins/Experimental/Animation/SkeletalMeshModelingTools/Source/SkeletalMeshModelingTools/Private/SkeletalMeshModelingToolsEditorMode.h b/Engine/Plugins/Experimental/Animation/SkeletalMeshModelingTools/Source/SkeletalMeshModelingTools/Private/SkeletalMeshModelingToolsEditorMode.h
new file mode 100644
--- /dev/null
+++ b/Engine/Plugins/Experimental/Animation/SkeletalMeshModelingTools/Source/SkeletalMeshModelingTools/Private/SkeletalMeshModelingToolsEditorMode.h
@@ -0,0 +1,43 @@
+// Copyright Epic Games, Inc. All Rights Reserved.
+
+#pragma once
+
+#include "CoreMinimal.h"
+#include "EdMode.h"
+
+#include "SkeletalMeshModelingToolsEditorMode.generated.h"
+
+
+class FStylusStateTracker;
+class FSkeletalMeshModelingToolsEditorModeToolkit;
+class UEdModeInteractiveToolsContext;
+
+
+UCLASS()
+class USkeletalMeshModelingToolsEditorMode : 
+	public UBaseLegacyWidgetEdMode
+{
+	GENERATED_BODY()
+public:
+	const static FEditorModeID Id;	
+
+	USkeletalMeshModelingToolsEditorMode();
+	explicit USkeletalMeshModelingToolsEditorMode(FVTableHelper& Helper);
+	virtual ~USkeletalMeshModelingToolsEditorMode() override;
+
+	// UEdMode overrides
+	void Initialize() override;
+
+	void Enter() override;
+	void Exit() override;
+	void CreateToolkit() override;
+
+	void Tick(FEditorViewportClient* InViewportClient, float InDeltaTime) override;
+	// void Render(const FSceneView* InView, FViewport* InViewport, FPrimitiveDrawInterface* InPDI) override;
+	// void DrawHUD(FEditorViewportClient* InViewportClient, FViewport* InViewport, const FSceneView* InView, FCanvas* InCanvas) override;
+
+	bool UsesToolkits() const override { return true; }
+
+private:
+	TUniquePtr<FStylusStateTracker> StylusStateTracker;
+};
diff --git a/Engine/Plugins/Experimental/Animation/SkeletalMeshModelingTools/Source/SkeletalMeshModelingTools/Private/SkeletalMeshModelingToolsEditorModeToolkit.cpp b/Engine/Plugins/Experimental/Animation/SkeletalMeshModelingTools/Source/SkeletalMeshModelingTools/Private/SkeletalMeshModelingToolsEditorModeToolkit.cpp
new file mode 100644
--- /dev/null
+++ b/Engine/Plugins/Experimental/Animation/SkeletalMeshModelingTools/Source/SkeletalMeshModelingTools/Private/SkeletalMeshModelingToolsEditorModeToolkit.cpp
@@ -0,0 +1,334 @@
+// Copyright Epic Games, Inc. All Rights Reserved.
+
+#include "SkeletalMeshModelingToolsEditorModeToolkit.h"
+
+#include "SkeletalMeshModelingToolsCommands.h"
+
+#include "EdModeInteractiveToolsContext.h"
+#include "Framework/MultiBox/MultiBoxBuilder.h"
+#include "IDetailsView.h"
+#include "InteractiveTool.h"
+#include "InteractiveToolsContext.h"
+#include "ModelingToolsEditorModeStyle.h"
+#include "Modules/ModuleManager.h"
+#include "PropertyEditorModule.h"
+#include "Tools/UEdMode.h"
+#include "Widgets/Input/SButton.h"
+#include "Widgets/Layout/SBorder.h"
+#include "Widgets/Images/SImage.h"
+#include "Widgets/Text/STextBlock.h"
+#include "ModelingToolsManagerActions.h"
+#include "BaseTools/MeshSurfacePointTool.h"
+#include "SPrimaryButton.h"
+
+
+#define LOCTEXT_NAMESPACE "SkeletalMeshModelingToolsEditorModeToolkit"
+
+
+void FSkeletalMeshModelingToolsEditorModeToolkit::Init(
+	const TSharedPtr<IToolkitHost>& InToolkitHost, 
+	TWeakObjectPtr<UEdMode> InOwningMode
+	)
+{
+	// Create a details view to show the tool props
+	FPropertyEditorModule& PropertyEditorModule = FModuleManager::GetModuleChecked<FPropertyEditorModule>("PropertyEditor");
+
+	FDetailsViewArgs DetailsViewArgs;
+
+	DetailsViewArgs.bHideSelectionTip = true;
+	DetailsViewArgs.NameAreaSettings = FDetailsViewArgs::HideNameArea;
+	DetailsViewArgs.bAllowSearch = false;
+	DetailsViewArgs.DefaultsOnlyVisibility = EEditDefaultsOnlyNodeVisibility::Automatic;
+	DetailsViewArgs.bShowOptions = false;
+	DetailsViewArgs.bAllowMultipleTopLevelObjects = true;
+
+	DetailsView = PropertyEditorModule.CreateDetailView(DetailsViewArgs);
+
+	ModeWarningArea = SNew(STextBlock)
+		.AutoWrapText(true)
+		.Font(FCoreStyle::GetDefaultFontStyle("Bold", 9))
+		.ColorAndOpacity(FSlateColor(FLinearColor(0.9f, 0.15f, 0.15f)));
+	ModeWarningArea->SetText(FText::GetEmpty());
+	ModeWarningArea->SetVisibility(EVisibility::Collapsed);
+
+	ModeHeaderArea = SNew(STextBlock)
+		.AutoWrapText(true)
+		.Font(FCoreStyle::GetDefaultFontStyle("Bold", 12));
+	ModeHeaderArea->SetText(LOCTEXT("SelectToolLabel", "Select a Tool from the Toolbar"));
+	ModeHeaderArea->SetJustification(ETextJustify::Center);
+
+
+	ToolWarningArea = SNew(STextBlock)
+		.AutoWrapText(true)
+		.Font(FCoreStyle::GetDefaultFontStyle("Bold", 9))
+		.ColorAndOpacity(FSlateColor(FLinearColor(0.9f, 0.15f, 0.15f)));
+	ToolWarningArea->SetText(FText::GetEmpty());
+
+
+	SAssignNew(ToolkitWidget, SBox)
+		.HAlign(HAlign_Fill)
+		.Padding(2)
+		[
+			SNew(SVerticalBox)
+			+ SVerticalBox::Slot().AutoHeight().HAlign(HAlign_Fill).Padding(5)
+			[
+				ModeWarningArea->AsShared()
+			]
+			+ SVerticalBox::Slot().AutoHeight().HAlign(HAlign_Fill).Padding(5)
+			[
+				ModeHeaderArea->AsShared()
+			]
+			+ SVerticalBox::Slot().AutoHeight().HAlign(HAlign_Fill).Padding(5)
+			[
+				ToolWarningArea->AsShared()
+			]
+			+ SVerticalBox::Slot().FillHeight(1.f).HAlign(HAlign_Fill)
+			[
+				DetailsView->AsShared()
+			]
+		];
+
+	FModeToolkit::Init(InToolkitHost, InOwningMode);
+
+	ClearNotification();
+	ClearWarning();
+
+	ActiveToolName = FText::GetEmpty();
+	ActiveToolMessage = FText::GetEmpty();
+	
+	GetScriptableEditorMode()->GetInteractiveToolsContext()->OnToolNotificationMessage.AddSP(this, &FSkeletalMeshModelingToolsEditorModeToolkit::PostNotification);
+	GetScriptableEditorMode()->GetInteractiveToolsContext()->OnToolWarningMessage.AddSP(this, &FSkeletalMeshModelingToolsEditorModeToolkit::PostWarning);
+
+	SAssignNew(ViewportOverlayWidget, SHorizontalBox)
+
+	+SHorizontalBox::Slot()
+	.HAlign(HAlign_Center)
+	.VAlign(VAlign_Bottom)
+	.Padding(FMargin(0.0f, 0.0f, 0.f, 15.f))
+	[
+		SNew(SBorder)
+		.BorderImage(FAppStyle::Get().GetBrush("EditorViewport.OverlayBrush"))
+		.Padding(8.f)
+		[
+			SNew(SHorizontalBox)
+
+			+SHorizontalBox::Slot()
+			.AutoWidth()
+			.VAlign(VAlign_Center)
+			.Padding(FMargin(0.f, 0.f, 8.f, 0.f))
+			[
+				SNew(SImage)
+				.Image_Lambda([this] () { return ActiveToolIcon; })
+			]
+
+			+SHorizontalBox::Slot()
+			.AutoWidth()
+			.VAlign(VAlign_Center)
+			.Padding(FMargin(0.f, 0.f, 8.f, 0.f))
+			[
+				SNew(STextBlock)
+				.Text(this, &FSkeletalMeshModelingToolsEditorModeToolkit::GetActiveToolDisplayName)
+			]
+
+			+SHorizontalBox::Slot()
+			.AutoWidth()
+			.Padding(FMargin(0.0, 0.f, 2.f, 0.f))
+			[
+				SNew(SPrimaryButton)
+				.Text(LOCTEXT("OverlayAccept", "Accept"))
+				.ToolTipText(LOCTEXT("OverlayAcceptTooltip", "Accept/Commit the results of the active Tool [Enter]"))
+				.OnClicked_Lambda([this]() { GetScriptableEditorMode()->GetInteractiveToolsContext()->EndTool(EToolShutdownType::Accept); return FReply::Handled(); })
+				.IsEnabled_Lambda([this]() { return GetScriptableEditorMode()->GetInteractiveToolsContext()->CanAcceptActiveTool(); })
+				.Visibility_Lambda([this]() { return GetScriptableEditorMode()->GetInteractiveToolsContext()->ActiveToolHasAccept() ? EVisibility::Visible : EVisibility::Collapsed; })
+			]
+
+			+SHorizontalBox::Slot()
+			.AutoWidth()
+			.Padding(FMargin(2.0, 0.f, 0.f, 0.f))
+			[
+				SNew(SButton)
+				.TextStyle( FAppStyle::Get(), "DialogButtonText" )
+				.Text(LOCTEXT("OverlayCancel", "Cancel"))
+				.ToolTipText(LOCTEXT("OverlayCancelTooltip", "Cancel the active Tool [Esc]"))
+				.HAlign(HAlign_Center)
+				.OnClicked_Lambda([this]() { GetScriptableEditorMode()->GetInteractiveToolsContext()->EndTool(EToolShutdownType::Cancel); return FReply::Handled(); })
+				.IsEnabled_Lambda([this]() { return GetScriptableEditorMode()->GetInteractiveToolsContext()->CanCancelActiveTool(); })
+				.Visibility_Lambda([this]() { return GetScriptableEditorMode()->GetInteractiveToolsContext()->ActiveToolHasAccept() ? EVisibility::Visible : EVisibility::Collapsed; })
+			]
+
+			+SHorizontalBox::Slot()
+			.AutoWidth()
+			.Padding(FMargin(2.0, 0.f, 0.f, 0.f))
+			[
+				SNew(SPrimaryButton)
+				.Text(LOCTEXT("OverlayComplete", "Complete"))
+				.ToolTipText(LOCTEXT("OverlayCompleteTooltip", "Exit the active Tool [Enter]"))
+				.OnClicked_Lambda([this]() { GetScriptableEditorMode()->GetInteractiveToolsContext()->EndTool(EToolShutdownType::Completed); return FReply::Handled(); })
+				.IsEnabled_Lambda([this]() { return GetScriptableEditorMode()->GetInteractiveToolsContext()->CanCompleteActiveTool(); })
+				.Visibility_Lambda([this]() { return GetScriptableEditorMode()->GetInteractiveToolsContext()->CanCompleteActiveTool() ? EVisibility::Visible : EVisibility::Collapsed; })
+			]
+		]	
+	];
+}
+
+
+FName FSkeletalMeshModelingToolsEditorModeToolkit::GetToolkitFName() const
+{
+	return FName("SkeletalMeshModelingToolsEditorModeToolkit");
+}
+
+
+FText FSkeletalMeshModelingToolsEditorModeToolkit::GetBaseToolkitName() const
+{
+	return LOCTEXT("DisplayName", "Skeletal Mesh Modeling Tools");
+}
+
+
+void FSkeletalMeshModelingToolsEditorModeToolkit::OnToolStarted(UInteractiveToolManager* Manager, UInteractiveTool* Tool)
+{
+	UpdateActiveToolProperties(Tool);
+
+	Tool->OnPropertySetsModified.AddSP(this, &FSkeletalMeshModelingToolsEditorModeToolkit::UpdateActiveToolProperties, Tool);
+
+	ModeHeaderArea->SetVisibility(EVisibility::Collapsed);
+	ActiveToolName = Tool->GetToolInfo().ToolDisplayName;
+
+	FString ActiveToolIdentifier = GetScriptableEditorMode()->GetToolManager()->GetActiveToolName(EToolSide::Left);
+	ActiveToolIdentifier.InsertAt(0, ".");
+	FName ActiveToolIconName = ISlateStyle::Join(FModelingToolsManagerCommands::Get().GetContextName(), TCHAR_TO_ANSI(*ActiveToolIdentifier));
+	ActiveToolIcon = FModelingToolsEditorModeStyle::Get()->GetOptionalBrush(ActiveToolIconName);
+	
+	GetToolkitHost()->AddViewportOverlayWidget(ViewportOverlayWidget.ToSharedRef());
+}
+
+
+void FSkeletalMeshModelingToolsEditorModeToolkit::OnToolEnded(UInteractiveToolManager* Manager, UInteractiveTool* Tool)
+{
+	if (IsHosted())
+	{
+		GetToolkitHost()->RemoveViewportOverlayWidget(ViewportOverlayWidget.ToSharedRef());
+	}
+	
+	if (Tool)
+	{
+		Tool->OnPropertySetsModified.RemoveAll(this);
+	}
+
+	ModeHeaderArea->SetVisibility(EVisibility::Visible);
+	DetailsView->SetObject(nullptr);
+	ActiveToolName = FText::GetEmpty();
+	ClearNotification();
+	ClearWarning();
+
+}
+
+
+FText FSkeletalMeshModelingToolsEditorModeToolkit::GetActiveToolDisplayName() const
+{
+	return ActiveToolName;
+}
+
+
+FText FSkeletalMeshModelingToolsEditorModeToolkit::GetActiveToolMessage() const
+{
+	return ActiveToolMessage;
+}
+
+static const FName EditTabName(TEXT("Edit"));
+static const FName ProcessingTabName(TEXT("MeshOps"));
+static const FName DeformTabName(TEXT("Deform"));
+static const FName SkinWeightsTabName(TEXT("Skin"));
+
+
+void FSkeletalMeshModelingToolsEditorModeToolkit::GetToolPaletteNames(TArray<FName>& InPaletteName) const
+{
+	InPaletteName = { EditTabName, ProcessingTabName, DeformTabName, SkinWeightsTabName };
+}
+
+
+FText FSkeletalMeshModelingToolsEditorModeToolkit::GetToolPaletteDisplayName(FName PaletteName) const
+{
+	return FText::FromName(PaletteName);
+}
+
+
+void FSkeletalMeshModelingToolsEditorModeToolkit::BuildToolPalette(FName PaletteName, class FToolBarBuilder& ToolbarBuilder)
+{
+	const FModelingToolsManagerCommands& Commands = FModelingToolsManagerCommands::Get();
+
+	if (PaletteName == EditTabName)
+	{
+		ToolbarBuilder.AddToolBarButton(Commands.BeginPolyEditTool);
+		ToolbarBuilder.AddToolBarButton(Commands.BeginPolyDeformTool);
+		ToolbarBuilder.AddToolBarButton(Commands.BeginHoleFillTool);
+		ToolbarBuilder.AddToolBarButton(Commands.BeginPolygonCutTool);
+	}
+	else if (PaletteName == ProcessingTabName)
+	{
+		ToolbarBuilder.AddToolBarButton(Commands.BeginSimplifyMeshTool);
+		ToolbarBuilder.AddToolBarButton(Commands.BeginRemeshMeshTool);
+		ToolbarBuilder.AddToolBarButton(Commands.BeginWeldEdgesTool);
+		ToolbarBuilder.AddToolBarButton(Commands.BeginRemoveOccludedTrianglesTool);
+		ToolbarBuilder.AddToolBarButton(Commands.BeginProjectToTargetTool);
+	}
+	else if (PaletteName == DeformTabName)
+	{
+		ToolbarBuilder.AddToolBarButton(Commands.BeginSculptMeshTool);
+		ToolbarBuilder.AddToolBarButton(Commands.BeginRemeshSculptMeshTool);
+		ToolbarBuilder.AddToolBarButton(Commands.BeginSmoothMeshTool);
+		ToolbarBuilder.AddToolBarButton(Commands.BeginOffsetMeshTool);
+		ToolbarBuilder.AddToolBarButton(Commands.BeginMeshSpaceDeformerTool);
+		ToolbarBuilder.AddToolBarButton(Commands.BeginLatticeDeformerTool);
+		ToolbarBuilder.AddToolBarButton(Commands.BeginDisplaceMeshTool);
+	}
+	else if (PaletteName == SkinWeightsTabName)
+	{
+		// TODO: Need attribute support on skelmeshes.
+		// ToolbarBuilder.AddToolBarButton(Commands.BeginMeshAttributePaintTool);
+		ToolbarBuilder.AddToolBarButton(Commands.BeginSkinWeightsPaintTool);
+		ToolbarBuilder.AddToolBarButton(Commands.BeginSkinWeightsBindingTool);
+	}
+}
+
+
+void FSkeletalMeshModelingToolsEditorModeToolkit::OnToolPaletteChanged(FName PaletteName)
+{
+}
+
+
+void FSkeletalMeshModelingToolsEditorModeToolkit::PostNotification(const FText& InMessage)
+{
+	ActiveToolMessage = InMessage;
+}
+
+
+void FSkeletalMeshModelingToolsEditorModeToolkit::ClearNotification()
+{
+	ActiveToolMessage = FText::GetEmpty();
+}
+
+
+void FSkeletalMeshModelingToolsEditorModeToolkit::PostWarning(const FText& Message)
+{
+	ToolWarningArea->SetText(Message);
+	ToolWarningArea->SetVisibility(EVisibility::Visible);
+}
+
+
+void FSkeletalMeshModelingToolsEditorModeToolkit::ClearWarning()
+{
+	ToolWarningArea->SetText(FText());
+	ToolWarningArea->SetVisibility(EVisibility::Collapsed);
+}
+
+
+void FSkeletalMeshModelingToolsEditorModeToolkit::UpdateActiveToolProperties(UInteractiveTool* Tool)
+{
+	if (Tool)
+	{
+		DetailsView->SetObjects(Tool->GetToolProperties(true));
+	}
+}
+
+
+#undef LOCTEXT_NAMESPACE
diff --git a/Engine/Plugins/Experimental/Animation/SkeletalMeshModelingTools/Source/SkeletalMeshModelingTools/Private/SkeletalMeshModelingToolsEditorModeToolkit.h b/Engine/Plugins/Experimental/Animation/SkeletalMeshModelingTools/Source/SkeletalMeshModelingTools/Private/SkeletalMeshModelingToolsEditorModeToolkit.h
new file mode 100644
--- /dev/null
+++ b/Engine/Plugins/Experimental/Animation/SkeletalMeshModelingTools/Source/SkeletalMeshModelingTools/Private/SkeletalMeshModelingToolsEditorModeToolkit.h
@@ -0,0 +1,57 @@
+// Copyright Epic Games, Inc. All Rights Reserved.
+
+#pragma once
+
+#include "CoreMinimal.h"
+#include "Toolkits/BaseToolkit.h"
+
+
+class IToolkitHost;
+class STextBlock;
+class SButton;
+
+
+class FSkeletalMeshModelingToolsEditorModeToolkit : 
+	 public FModeToolkit
+{
+public:
+	// IToolkit overrides
+	void Init(const TSharedPtr<IToolkitHost>& InToolkitHost, TWeakObjectPtr<UEdMode> InOwningMode) override;
+	FName GetToolkitFName() const override;
+	FText GetBaseToolkitName() const override;
+	TSharedPtr<SWidget> GetInlineContent() const override { return ToolkitWidget; }
+
+	void OnToolStarted(UInteractiveToolManager* Manager, UInteractiveTool* Tool) override;
+	void OnToolEnded(UInteractiveToolManager* Manager, UInteractiveTool* Tool) override;
+
+	FText GetActiveToolDisplayName() const override;
+	FText GetActiveToolMessage() const override;
+
+	void GetToolPaletteNames(TArray<FName>& InPaletteName) const override;
+	FText GetToolPaletteDisplayName(FName PaletteName) const override;
+	void BuildToolPalette(FName PaletteName, class FToolBarBuilder& ToolbarBuilder) override;
+	void OnToolPaletteChanged(FName PaletteName) override;
+
+
+private:
+	void PostNotification(const FText& Message);
+	void ClearNotification();
+
+	void PostWarning(const FText& Message);
+	void ClearWarning();
+
+	void UpdateActiveToolProperties(UInteractiveTool* Tool);
+
+	FText ActiveToolName;
+	FText ActiveToolMessage;
+
+	TSharedPtr<SWidget> ToolkitWidget;
+	TSharedPtr<IDetailsView> DetailsView;
+
+	TSharedPtr<SWidget> ViewportOverlayWidget;
+	const FSlateBrush* ActiveToolIcon = nullptr;
+	
+	TSharedPtr<STextBlock> ModeWarningArea;
+	TSharedPtr<STextBlock> ModeHeaderArea;
+	TSharedPtr<STextBlock> ToolWarningArea;
+};
diff --git a/Engine/Plugins/Experimental/Animation/SkeletalMeshModelingTools/Source/SkeletalMeshModelingTools/Private/SkeletalMeshModelingToolsModule.cpp b/Engine/Plugins/Experimental/Animation/SkeletalMeshModelingTools/Source/SkeletalMeshModelingTools/Private/SkeletalMeshModelingToolsModule.cpp
new file mode 100644
--- /dev/null
+++ b/Engine/Plugins/Experimental/Animation/SkeletalMeshModelingTools/Source/SkeletalMeshModelingTools/Private/SkeletalMeshModelingToolsModule.cpp
@@ -0,0 +1,120 @@
+// Copyright Epic Games, Inc. All Rights Reserved.
+
+#include "SkeletalMeshModelingToolsModule.h"
+
+#include "SkeletalMeshModelingToolsCommands.h"
+#include "SkeletalMeshModelingToolsEditorMode.h"
+#include "SkeletalMeshModelingToolsStyle.h"
+
+#include "EditorModeManager.h"
+#include "EditorModeRegistry.h"
+#include "ISkeletalMeshEditor.h"
+#include "ISkeletalMeshEditorModule.h"
+#include "Modules/ModuleManager.h"
+#include "SkeletalMeshToolMenuContext.h"
+#include "ToolMenus.h"
+#include "WorkflowOrientedApp/ApplicationMode.h"
+
+
+#define LOCTEXT_NAMESPACE "SkeletalMeshModelingToolsModule"
+
+DEFINE_LOG_CATEGORY(LogSkeletalMeshModelingTools);
+
+IMPLEMENT_MODULE(FSkeletalMeshModelingToolsModule, SkeletalMeshModelingTools);
+
+
+
+void FSkeletalMeshModelingToolsModule::StartupModule()
+{
+	FSkeletalMeshModelingToolsStyle::Register();
+	FSkeletalMeshModelingToolsCommands::Register();
+
+	UToolMenus::RegisterStartupCallback(
+		FSimpleMulticastDelegate::FDelegate::CreateRaw(this, &FSkeletalMeshModelingToolsModule::RegisterMenusAndToolbars));
+
+	ISkeletalMeshEditorModule& SkelMeshEditorModule = FModuleManager::Get().LoadModuleChecked<ISkeletalMeshEditorModule>("SkeletalMeshEditor");
+	TArray<ISkeletalMeshEditorModule::FSkeletalMeshEditorToolbarExtender>& ToolbarExtenders = SkelMeshEditorModule.GetAllSkeletalMeshEditorToolbarExtenders();
+
+	ToolbarExtenders.Add(ISkeletalMeshEditorModule::FSkeletalMeshEditorToolbarExtender::CreateRaw(this, &FSkeletalMeshModelingToolsModule::ExtendSkelMeshEditorToolbar));
+	SkelMeshEditorExtenderHandle = ToolbarExtenders.Last().GetHandle();
+}
+
+void FSkeletalMeshModelingToolsModule::ShutdownModule()
+{
+	if (ISkeletalMeshEditorModule* SkelMeshEditorModule = FModuleManager::GetModulePtr<ISkeletalMeshEditorModule>("SkeletalMeshEditor"))
+	{
+		TArray<ISkeletalMeshEditorModule::FSkeletalMeshEditorToolbarExtender>& Extenders = SkelMeshEditorModule->GetAllSkeletalMeshEditorToolbarExtenders();
+
+		Extenders.RemoveAll([=](const ISkeletalMeshEditorModule::FSkeletalMeshEditorToolbarExtender& InDelegate) { return InDelegate.GetHandle() == SkelMeshEditorExtenderHandle; });
+	}
+
+	UToolMenus::UnregisterOwner(this);
+
+	FSkeletalMeshModelingToolsCommands::Unregister();
+	FSkeletalMeshModelingToolsStyle::Unregister();
+}
+
+
+void FSkeletalMeshModelingToolsModule::RegisterMenusAndToolbars()
+{
+	FToolMenuOwnerScoped OwnerScoped(this);
+
+	UToolMenu* Toolbar = UToolMenus::Get()->ExtendMenu("AssetEditor.SkeletalMeshEditor.ToolBar");
+	{
+		FToolMenuSection& Section = Toolbar->FindOrAddSection("SkeletalMesh");
+		Section.AddDynamicEntry("ToggleModelingToolsMode", FNewToolMenuSectionDelegate::CreateLambda([this](FToolMenuSection& InSection) {
+			USkeletalMeshToolMenuContext* Context = InSection.FindContext<USkeletalMeshToolMenuContext>();
+			if (Context && Context->SkeletalMeshEditor.IsValid())
+			{
+				InSection.AddEntry(FToolMenuEntry::InitToolBarButton(
+				    FSkeletalMeshModelingToolsCommands::Get().ToggleModelingToolsMode,
+					LOCTEXT("SkeletalMeshEditorModelingMode", "Modeling Tools"),
+				    LOCTEXT("SkeletalMeshEditorModelingModeTooltip", "Opens the Modeling Tools palette that provides selected mesh modification tools."),
+				    FSlateIcon("ModelingToolsStyle", "LevelEditor.ModelingToolsMode")));
+			}
+		}));
+	}
+}
+
+
+TSharedRef<FExtender> FSkeletalMeshModelingToolsModule::ExtendSkelMeshEditorToolbar(const TSharedRef<FUICommandList> InCommandList, TSharedRef<ISkeletalMeshEditor> InSkeletalMeshEditor)
+{
+	// Add toolbar extender
+	TSharedPtr<FExtender> ToolbarExtender = MakeShareable(new FExtender);
+
+	TWeakPtr<ISkeletalMeshEditor> Ptr(InSkeletalMeshEditor);
+
+	InCommandList->MapAction(FSkeletalMeshModelingToolsCommands::Get().ToggleModelingToolsMode,
+	    FExecuteAction::CreateRaw(this, &FSkeletalMeshModelingToolsModule::OnToggleModelingToolsMode, Ptr),
+	    FCanExecuteAction(),
+	    FIsActionChecked::CreateRaw(this, &FSkeletalMeshModelingToolsModule::IsModelingToolModeActive, Ptr));
+
+	return ToolbarExtender.ToSharedRef();
+}
+
+
+bool FSkeletalMeshModelingToolsModule::IsModelingToolModeActive(TWeakPtr<ISkeletalMeshEditor> InSkeletalMeshEditor) const
+{
+	TSharedPtr<ISkeletalMeshEditor> SkeletalMeshEditor = InSkeletalMeshEditor.Pin();
+	return SkeletalMeshEditor.IsValid() && SkeletalMeshEditor->GetEditorModeManager().IsModeActive(USkeletalMeshModelingToolsEditorMode::Id);
+}
+
+
+void FSkeletalMeshModelingToolsModule::OnToggleModelingToolsMode(TWeakPtr<ISkeletalMeshEditor> InSkeletalMeshEditor)
+{
+	TSharedPtr<ISkeletalMeshEditor> SkeletalMeshEditor = InSkeletalMeshEditor.Pin();
+	if (SkeletalMeshEditor.IsValid())
+	{
+		if (!IsModelingToolModeActive(InSkeletalMeshEditor))
+		{
+			SkeletalMeshEditor->GetEditorModeManager().ActivateMode(USkeletalMeshModelingToolsEditorMode::Id, true);
+		}
+		else
+		{
+			SkeletalMeshEditor->GetEditorModeManager().ActivateDefaultMode();
+		}
+	}
+}
+
+
+#undef LOCTEXT_NAMESPACE
diff --git a/Engine/Plugins/Experimental/Animation/SkeletalMeshModelingTools/Source/SkeletalMeshModelingTools/Private/SkeletalMeshModelingToolsModule.h b/Engine/Plugins/Experimental/Animation/SkeletalMeshModelingTools/Source/SkeletalMeshModelingTools/Private/SkeletalMeshModelingToolsModule.h
new file mode 100644
--- /dev/null
+++ b/Engine/Plugins/Experimental/Animation/SkeletalMeshModelingTools/Source/SkeletalMeshModelingTools/Private/SkeletalMeshModelingToolsModule.h
@@ -0,0 +1,39 @@
+// Copyright Epic Games, Inc. All Rights Reserved.
+
+#pragma once
+
+#include "CoreMinimal.h"
+#include "Modules/ModuleInterface.h"
+
+
+DECLARE_LOG_CATEGORY_EXTERN(LogSkeletalMeshModelingTools, Log, All);
+
+
+class FApplicationMode;
+class ISkeletalMeshEditor;
+class FUICommandList;
+class FExtender;
+
+class FSkeletalMeshModelingToolsModule : 
+	public IModuleInterface
+{
+public:
+	// IModuleInterface implementations
+	void StartupModule() override;
+	void ShutdownModule() override;
+
+private:
+	// Extend toolbars and menus
+	void RegisterMenusAndToolbars();
+
+	TSharedRef<FExtender> ExtendSkelMeshEditorToolbar(const TSharedRef<FUICommandList> InCommandList, TSharedRef<ISkeletalMeshEditor> InSkeletalMeshEditor);
+
+
+	bool IsModelingToolModeActive(TWeakPtr<ISkeletalMeshEditor> InSkeletalMeshEditor) const;
+	void OnToggleModelingToolsMode(TWeakPtr<ISkeletalMeshEditor> InSkeletalMeshEditor);
+
+	// The handle for the extender delegate we added. Needed for clean module shutdown.
+	FDelegateHandle SkelMeshEditorExtenderHandle;
+
+	TArray<TWeakPtr<FApplicationMode>> RegisteredApplicationModes;
+};
diff --git a/Engine/Plugins/Experimental/Animation/SkeletalMeshModelingTools/Source/SkeletalMeshModelingTools/Private/SkeletalMeshModelingToolsStyle.cpp b/Engine/Plugins/Experimental/Animation/SkeletalMeshModelingTools/Source/SkeletalMeshModelingTools/Private/SkeletalMeshModelingToolsStyle.cpp
new file mode 100644
--- /dev/null
+++ b/Engine/Plugins/Experimental/Animation/SkeletalMeshModelingTools/Source/SkeletalMeshModelingTools/Private/SkeletalMeshModelingToolsStyle.cpp
@@ -0,0 +1,32 @@
+// Copyright Epic Games, Inc. All Rights Reserved.
+
+#include "SkeletalMeshModelingToolsStyle.h"
+
+#include "Styling/SlateStyleMacros.h"
+#include "Styling/SlateStyleRegistry.h"
+
+
+FSkeletalMeshModelingToolsStyle::FSkeletalMeshModelingToolsStyle() : 
+	FSlateStyleSet("SkeletalMeshModelingToolsStyle")
+{
+	
+}
+
+
+FSkeletalMeshModelingToolsStyle& FSkeletalMeshModelingToolsStyle::Get()
+{
+	static FSkeletalMeshModelingToolsStyle Instance;
+	return Instance;
+}
+
+
+void FSkeletalMeshModelingToolsStyle::Register()
+{
+	FSlateStyleRegistry::RegisterSlateStyle(Get());
+}
+
+
+void FSkeletalMeshModelingToolsStyle::Unregister()
+{
+	FSlateStyleRegistry::UnRegisterSlateStyle(Get());
+}
diff --git a/Engine/Plugins/Experimental/Animation/SkeletalMeshModelingTools/Source/SkeletalMeshModelingTools/Private/SkeletalMeshModelingToolsStyle.h b/Engine/Plugins/Experimental/Animation/SkeletalMeshModelingTools/Source/SkeletalMeshModelingTools/Private/SkeletalMeshModelingToolsStyle.h
new file mode 100644
--- /dev/null
+++ b/Engine/Plugins/Experimental/Animation/SkeletalMeshModelingTools/Source/SkeletalMeshModelingTools/Private/SkeletalMeshModelingToolsStyle.h
@@ -0,0 +1,22 @@
+// Copyright Epic Games, Inc. All Rights Reserved.
+
+#pragma once
+
+#include "Styling/SlateStyle.h"
+
+
+class FSkeletalMeshModelingToolsStyle
+    : public FSlateStyleSet
+{
+public:
+	static FSkeletalMeshModelingToolsStyle& Get();
+
+protected:
+	friend class FSkeletalMeshModelingToolsModule;
+
+	static void Register();
+	static void Unregister();
+
+private:
+	FSkeletalMeshModelingToolsStyle();
+};
diff --git a/Engine/Plugins/Experimental/Animation/SkeletalMeshModelingTools/Source/SkeletalMeshModelingTools/SkeletalMeshModelingTools.Build.cs b/Engine/Plugins/Experimental/Animation/SkeletalMeshModelingTools/Source/SkeletalMeshModelingTools/SkeletalMeshModelingTools.Build.cs
new file mode 100644
--- /dev/null
+++ b/Engine/Plugins/Experimental/Animation/SkeletalMeshModelingTools/Source/SkeletalMeshModelingTools/SkeletalMeshModelingTools.Build.cs
@@ -0,0 +1,40 @@
+﻿// Copyright Epic Games, Inc. All Rights Reserved.
+
+using UnrealBuildTool;
+
+public class SkeletalMeshModelingTools : ModuleRules
+{
+	public SkeletalMeshModelingTools(ReadOnlyTargetRules Target) : base(Target)
+	{
+		PrivateIncludePaths.AddRange(
+			new string[] {
+				"SkeletalMeshModelingTools/Private",
+			}
+		);
+
+		PrivateDependencyModuleNames.AddRange(
+			new string[]
+			{
+				"Core",
+				"CoreUObject",
+				"Slate",
+				"SlateCore",
+				"EditorFramework",
+				"Engine",
+				"ToolMenus",
+				"UnrealEd",
+				"InputCore",
+				"ModelingComponentsEditorOnly",
+				"MeshModelingTools",
+				"MeshModelingToolsExp",
+				"MeshModelingToolsEditorOnly",
+				"MeshModelingToolsEditorOnlyExp",
+				"ModelingToolsEditorMode",
+				"InteractiveToolsFramework",
+				"EditorInteractiveToolsFramework",
+				"StylusInput",
+				"ToolWidgets",
+			}
+		);
+	}
+}
