Add missing #include to fix non-unity builds.

#rnx

diff --git a/Engine/Plugins/Experimental/Animation/SkeletalMeshModelingTools/Source/SkeletalMeshModelingTools/Private/SkeletalMeshModelingToolsModule.cpp b/Engine/Plugins/Experimental/Animation/SkeletalMeshModelingTools/Source/SkeletalMeshModelingTools/Private/SkeletalMeshModelingToolsModule.cpp
--- a/Engine/Plugins/Experimental/Animation/SkeletalMeshModelingTools/Source/SkeletalMeshModelingTools/Private/SkeletalMeshModelingToolsModule.cpp
+++ b/Engine/Plugins/Experimental/Animation/SkeletalMeshModelingTools/Source/SkeletalMeshModelingTools/Private/SkeletalMeshModelingToolsModule.cpp
@@ -9,7 +9,7 @@
 
 #include "ContentBrowserMenuContexts.h"
 #include "EditorModeManager.h"
-#include "EditorModeRegistry.h"
+#include "Engine/SkeletalMesh.h"
 #include "ISkeletalMeshEditor.h"
 #include "ISkeletalMeshEditorModule.h"
 #include "Modules/ModuleManager.h"
