Fix crash on editor shutdown with GEditor->GetEditorSubsystem<UStylusInputSubsystem>()

diff --git a/Engine/Plugins/Editor/ModelingToolsEditorMode/Source/ModelingToolsEditorMode/Private/ModelingToolsEditorMode.cpp b/Engine/Plugins/Editor/ModelingToolsEditorMode/Source/ModelingToolsEditorMode/Private/ModelingToolsEditorMode.cpp
--- a/Engine/Plugins/Editor/ModelingToolsEditorMode/Source/ModelingToolsEditorMode/Private/ModelingToolsEditorMode.cpp
+++ b/Engine/Plugins/Editor/ModelingToolsEditorMode/Source/ModelingToolsEditorMode/Private/ModelingToolsEditorMode.cpp
@@ -287,8 +287,13 @@ public:
 
 	virtual ~FStylusStateTracker()
 	{
-		UStylusInputSubsystem* StylusSubsystem = GEditor->GetEditorSubsystem<UStylusInputSubsystem>();
-		StylusSubsystem->RemoveMessageHandler(*this);
+		if (GEditor)
+		{
+			if (UStylusInputSubsystem* StylusSubsystem = GEditor->GetEditorSubsystem<UStylusInputSubsystem>())
+			{
+				StylusSubsystem->RemoveMessageHandler(*this);
+			}
+		}
 	}
 
 	virtual void OnStylusStateChanged(const FStylusState& NewState, int32 StylusIndex) override
diff --git a/Engine/Plugins/Experimental/Animation/SkeletalMeshModelingTools/Source/SkeletalMeshModelingTools/Private/SkeletalMeshModelingToolsEditorMode.cpp b/Engine/Plugins/Experimental/Animation/SkeletalMeshModelingTools/Source/SkeletalMeshModelingTools/Private/SkeletalMeshModelingToolsEditorMode.cpp
--- a/Engine/Plugins/Experimental/Animation/SkeletalMeshModelingTools/Source/SkeletalMeshModelingTools/Private/SkeletalMeshModelingToolsEditorMode.cpp
+++ b/Engine/Plugins/Experimental/Animation/SkeletalMeshModelingTools/Source/SkeletalMeshModelingTools/Private/SkeletalMeshModelingToolsEditorMode.cpp
@@ -67,8 +67,13 @@ public:
 
 	virtual ~FStylusStateTracker()
 	{
-		UStylusInputSubsystem* StylusSubsystem = GEditor->GetEditorSubsystem<UStylusInputSubsystem>();
-		StylusSubsystem->RemoveMessageHandler(*this);
+		if (GEditor)
+		{
+			if (UStylusInputSubsystem* StylusSubsystem = GEditor->GetEditorSubsystem<UStylusInputSubsystem>())
+			{
+				StylusSubsystem->RemoveMessageHandler(*this);
+			}
+		}
 	}
 
 	void OnStylusStateChanged(const FStylusState& NewState, int32 StylusIndex) override
