Skeletal Mesh Editing interface (submitted on behalf of Benoit Gadreau)

diff --git a/Engine/Plugins/Experimental/Animation/SkeletalMeshModelingTools/Source/SkeletalMeshModelingTools/Private/SkeletalMeshModelingToolsEditorMode.cpp b/Engine/Plugins/Experimental/Animation/SkeletalMeshModelingTools/Source/SkeletalMeshModelingTools/Private/SkeletalMeshModelingToolsEditorMode.cpp
--- a/Engine/Plugins/Experimental/Animation/SkeletalMeshModelingTools/Source/SkeletalMeshModelingTools/Private/SkeletalMeshModelingToolsEditorMode.cpp
+++ b/Engine/Plugins/Experimental/Animation/SkeletalMeshModelingTools/Source/SkeletalMeshModelingTools/Private/SkeletalMeshModelingToolsEditorMode.cpp
@@ -36,6 +36,8 @@
 #include "ToolTargetManager.h"
 #include "WeldMeshEdgesTool.h"
 
+#include "SkeletalMeshNotifier.h"
+#include "SkeletalMesh/SkeletalMeshEditionInterface.h"
 
 #define LOCTEXT_NAMESPACE "SkeletalMeshModelingToolsEditorMode"
 
@@ -249,5 +251,106 @@ void USkeletalMeshModelingToolsEditorMode::Tick(FEditorViewportClient* ViewportC
 	Super::Tick(ViewportClient, DeltaTime);
 }
 
+bool USkeletalMeshModelingToolsEditorMode::HandleClick(FEditorViewportClient* InViewportClient, HHitProxy *HitProxy, const FViewportClick &Click)
+{
+	if (Binding.IsValid())
+	{
+		TSharedPtr<ISkeletalMeshEditorBinding> BindingPtr = Binding.Pin();
+
+		TArray<FName> SelectedBones;
+		if (HitProxy && BindingPtr->GetNameFunction())
+		{
+			if (TOptional<FName> BoneName = BindingPtr->GetNameFunction()(HitProxy))
+			{
+				SelectedBones.Emplace(*BoneName);
+			}
+		}
+		
+		BindingPtr->GetNotifier().HandleNotification( SelectedBones, ESkeletalMeshNotifyType::BonesSelected);
+	}
+	
+	return Super::HandleClick(InViewportClient, HitProxy, Click);
+}
+
+void USkeletalMeshModelingToolsEditorMode::OnToolStarted(UInteractiveToolManager* Manager, UInteractiveTool* Tool)
+{
+	ConnectTool(Tool);
+}
+
+void USkeletalMeshModelingToolsEditorMode::OnToolEnded(UInteractiveToolManager* Manager, UInteractiveTool* Tool)
+{
+	DisconnectTool(Tool);
+}
+
+void USkeletalMeshModelingToolsEditorMode::SetEditorBinding(TSharedPtr<ISkeletalMeshEditorBinding> InBinding)
+{
+	Binding = InBinding;
+}
+
+ISkeletalMeshEditionInterface* USkeletalMeshModelingToolsEditorMode::GetSkeletonInterface(UInteractiveTool* InTool)
+{
+	if (!IsValid(InTool) && !InTool->Implements<USkeletalMeshEditionInterface>())
+	{
+		return nullptr;
+	}
+	return static_cast<ISkeletalMeshEditionInterface*>(InTool->GetInterfaceAddress(USkeletalMeshEditionInterface::StaticClass()));
+}
+
+void USkeletalMeshModelingToolsEditorMode::ConnectTool(UInteractiveTool* InTool)
+{
+	ISkeletalMeshEditionInterface* SkeletonInterface = GetSkeletonInterface(InTool);
+	if (!SkeletonInterface)
+	{
+		return;
+	}
+
+	// binding
+	if (Binding.IsValid())
+	{
+		TSharedPtr<ISkeletalMeshEditorBinding> BindingPtr = Binding.Pin();
+
+		SkeletonInterface->BindTo(BindingPtr);
+
+		// connect external interface to tool (ie skeletal mesh editor -> tool)
+		if (!ToToolNotifierHandle.IsValid())
+		{
+			ToToolNotifierHandle = BindingPtr->GetNotifier().Delegate().AddLambda([SkeletonInterface](const TArray<FName>& BoneNames, const ESkeletalMeshNotifyType InNotifyType)
+			{
+				SkeletonInterface->GetNotifier().HandleNotification(BoneNames, InNotifyType);
+			});
+		}
+
+		// connect too to external interface (ie tool -> skeletal mesh editor)
+		if (!FromToolNotifierHandle.IsValid())
+		{
+			FromToolNotifierHandle = SkeletonInterface->GetNotifier().Delegate().AddLambda([BindingPtr](const TArray<FName>& BoneNames, const ESkeletalMeshNotifyType InNotifyType)
+			{
+				BindingPtr->GetNotifier().HandleNotification(BoneNames, InNotifyType);
+			});
+		}
+	}
+}
+
+void USkeletalMeshModelingToolsEditorMode::DisconnectTool(UInteractiveTool* InTool)
+{
+	if (ToToolNotifierHandle.IsValid())
+	{
+		// check that
+		if (Binding.IsValid())
+		{
+			Binding.Pin()->GetNotifier().Delegate().Remove(ToToolNotifierHandle);
+		}
+		ToToolNotifierHandle.Reset();
+	}
+
+	if (FromToolNotifierHandle.IsValid())
+	{
+		ISkeletalMeshEditionInterface* SkeletonInterface = GetSkeletonInterface(InTool);
+		check(SkeletonInterface);
+		SkeletonInterface->GetNotifier().Delegate().Remove(FromToolNotifierHandle);
+		FromToolNotifierHandle.Reset();
+		SkeletonInterface->Unbind();
+	}
+}
 
 #undef LOCTEXT_NAMESPACE
diff --git a/Engine/Plugins/Experimental/Animation/SkeletalMeshModelingTools/Source/SkeletalMeshModelingTools/Private/SkeletalMeshModelingToolsEditorMode.h b/Engine/Plugins/Experimental/Animation/SkeletalMeshModelingTools/Source/SkeletalMeshModelingTools/Private/SkeletalMeshModelingToolsEditorMode.h
--- a/Engine/Plugins/Experimental/Animation/SkeletalMeshModelingTools/Source/SkeletalMeshModelingTools/Private/SkeletalMeshModelingToolsEditorMode.h
+++ b/Engine/Plugins/Experimental/Animation/SkeletalMeshModelingTools/Source/SkeletalMeshModelingTools/Private/SkeletalMeshModelingToolsEditorMode.h
@@ -7,11 +7,13 @@
 
 #include "SkeletalMeshModelingToolsEditorMode.generated.h"
 
-
 class FStylusStateTracker;
 class FSkeletalMeshModelingToolsEditorModeToolkit;
 class UEdModeInteractiveToolsContext;
-
+class ISkeletalMeshNotifier;
+class ISkeletalMeshEditorBinding;
+class ISkeletalMeshEditionInterface;
+class HHitProxy;
 
 UCLASS()
 class USkeletalMeshModelingToolsEditorMode : 
@@ -26,18 +28,36 @@ public:
 	virtual ~USkeletalMeshModelingToolsEditorMode() override;
 
 	// UEdMode overrides
-	void Initialize() override;
+	virtual void Initialize() override;
 
-	void Enter() override;
-	void Exit() override;
-	void CreateToolkit() override;
+	virtual void Enter() override;
+	virtual void Exit() override;
+	virtual void CreateToolkit() override;
 
-	void Tick(FEditorViewportClient* InViewportClient, float InDeltaTime) override;
+	virtual void Tick(FEditorViewportClient* InViewportClient, float InDeltaTime) override;
 	// void Render(const FSceneView* InView, FViewport* InViewport, FPrimitiveDrawInterface* InPDI) override;
 	// void DrawHUD(FEditorViewportClient* InViewportClient, FViewport* InViewport, const FSceneView* InView, FCanvas* InCanvas) override;
+	virtual bool HandleClick(FEditorViewportClient* InViewportClient, HHitProxy *HitProxy, const FViewportClick &Click) override;
 
-	bool UsesToolkits() const override { return true; }
+	virtual bool UsesToolkits() const override { return true; }
 
+	// binding
+	void SetEditorBinding(TSharedPtr<ISkeletalMeshEditorBinding> InBinding);
+		
+protected:
+	virtual void OnToolStarted(UInteractiveToolManager* Manager, UInteractiveTool* Tool) override;
+	virtual void OnToolEnded(UInteractiveToolManager* Manager, UInteractiveTool* Tool) override;
+	
 private:
 	TUniquePtr<FStylusStateTracker> StylusStateTracker;
+
+	static ISkeletalMeshEditionInterface* GetSkeletonInterface(UInteractiveTool* InTool);
+
+	void ConnectTool(UInteractiveTool* InTool);
+	void DisconnectTool(UInteractiveTool* InTool);
+
+	FDelegateHandle ToToolNotifierHandle;
+	FDelegateHandle FromToolNotifierHandle;
+	
+	TWeakPtr<ISkeletalMeshEditorBinding> Binding;
 };
diff --git a/Engine/Plugins/Experimental/Animation/SkeletalMeshModelingTools/Source/SkeletalMeshModelingTools/Private/SkeletalMeshModelingToolsModule.cpp b/Engine/Plugins/Experimental/Animation/SkeletalMeshModelingTools/Source/SkeletalMeshModelingTools/Private/SkeletalMeshModelingToolsModule.cpp
--- a/Engine/Plugins/Experimental/Animation/SkeletalMeshModelingTools/Source/SkeletalMeshModelingTools/Private/SkeletalMeshModelingToolsModule.cpp
+++ b/Engine/Plugins/Experimental/Animation/SkeletalMeshModelingTools/Source/SkeletalMeshModelingTools/Private/SkeletalMeshModelingToolsModule.cpp
@@ -130,16 +130,22 @@ void FSkeletalMeshModelingToolsModule::OnToggleModelingToolsMode(TWeakPtr<ISkele
 	TSharedPtr<ISkeletalMeshEditor> SkeletalMeshEditor = InSkeletalMeshEditor.Pin();
 	if (SkeletalMeshEditor.IsValid())
 	{
+		FEditorModeTools& EditorModeManager = SkeletalMeshEditor->GetEditorModeManager();
 		if (!IsModelingToolModeActive(InSkeletalMeshEditor))
 		{
-			SkeletalMeshEditor->GetEditorModeManager().ActivateMode(USkeletalMeshModelingToolsEditorMode::Id, true);
+			EditorModeManager.ActivateMode(USkeletalMeshModelingToolsEditorMode::Id, true);
+
+			UEdMode* ActiveMode = EditorModeManager.GetActiveScriptableMode(USkeletalMeshModelingToolsEditorMode::Id);
+			if (USkeletalMeshModelingToolsEditorMode* SkeletalMeshMode = CastChecked<USkeletalMeshModelingToolsEditorMode>(ActiveMode))
+			{
+				SkeletalMeshMode->SetEditorBinding(SkeletalMeshEditor->GetBinding());
+			}
 		}
 		else
 		{
-			SkeletalMeshEditor->GetEditorModeManager().ActivateDefaultMode();
+			EditorModeManager.ActivateDefaultMode();
 		}
 	}
 }
 
-
 #undef LOCTEXT_NAMESPACE
diff --git a/Engine/Plugins/Experimental/MeshModelingToolsetExp/Source/MeshModelingToolsExp/MeshModelingToolsExp.Build.cs b/Engine/Plugins/Experimental/MeshModelingToolsetExp/Source/MeshModelingToolsExp/MeshModelingToolsExp.Build.cs
--- a/Engine/Plugins/Experimental/MeshModelingToolsetExp/Source/MeshModelingToolsExp/MeshModelingToolsExp.Build.cs
+++ b/Engine/Plugins/Experimental/MeshModelingToolsetExp/Source/MeshModelingToolsExp/MeshModelingToolsExp.Build.cs
@@ -37,6 +37,7 @@ public class MeshModelingToolsExp : ModuleRules
 				"SkeletalMeshDescription",
 				"ModelingComponents",
 				"ModelingOperators",
+				"SkeletalMeshUtilitiesCommon",
 				// ... add other public dependencies that you statically link with here ...
 			}
 		);
diff --git a/Engine/Plugins/Experimental/MeshModelingToolsetExp/Source/MeshModelingToolsExp/Private/SkeletalMesh/SkeletalMeshEditionInterface.cpp b/Engine/Plugins/Experimental/MeshModelingToolsetExp/Source/MeshModelingToolsExp/Private/SkeletalMesh/SkeletalMeshEditionInterface.cpp
new file mode 100644
--- /dev/null
+++ b/Engine/Plugins/Experimental/MeshModelingToolsetExp/Source/MeshModelingToolsExp/Private/SkeletalMesh/SkeletalMeshEditionInterface.cpp
@@ -0,0 +1,53 @@
+﻿// Copyright Epic Games, Inc. All Rights Reserved.
+
+#include "SkeletalMesh/SkeletalMeshEditionInterface.h"
+
+IMPLEMENT_HIT_PROXY(HBoneHitProxy, HHitProxy)
+
+ISkeletalMeshNotifier& ISkeletalMeshEditionInterface::GetNotifier()
+{
+	if (!Notifier)
+	{
+		Notifier.Reset(new FSkeletalMeshToolNotifier(TWeakInterfacePtr<ISkeletalMeshEditionInterface>(this)));
+	}
+	
+	return *Notifier;
+}
+
+bool ISkeletalMeshEditionInterface::NeedsNotification() const
+{
+	return Notifier && Notifier->Delegate().IsBound(); 
+}
+
+void ISkeletalMeshEditionInterface::BindTo(TSharedPtr<ISkeletalMeshEditorBinding> InBinding)
+{
+	Binding = InBinding;
+}
+
+void ISkeletalMeshEditionInterface::Unbind()
+{
+	Binding.Reset();
+}
+
+TOptional<FName> ISkeletalMeshEditionInterface::GetBoneName(HHitProxy* InHitProxy) const
+{
+	if (const HBoneHitProxy* BoneProxy = HitProxyCast<HBoneHitProxy>(InHitProxy))
+	{
+		return BoneProxy->BoneName;
+	}
+
+	return Binding.IsValid() && Binding.Pin()->GetNameFunction() ? Binding.Pin()->GetNameFunction()(InHitProxy) : TOptional<FName>();
+}
+
+FSkeletalMeshToolNotifier::FSkeletalMeshToolNotifier(TWeakInterfacePtr<ISkeletalMeshEditionInterface> InInterface)
+	: ISkeletalMeshNotifier()
+	, Interface(InInterface)
+{}
+
+void FSkeletalMeshToolNotifier::HandleNotification(const TArray<FName>& BoneNames, const ESkeletalMeshNotifyType InNotifyType)
+{
+	if (Interface.IsValid())
+	{
+		Interface->HandleSkeletalMeshModified(BoneNames, InNotifyType);
+	}
+}
diff --git a/Engine/Plugins/Experimental/MeshModelingToolsetExp/Source/MeshModelingToolsExp/Private/SkinWeightsPaintTool.cpp b/Engine/Plugins/Experimental/MeshModelingToolsetExp/Source/MeshModelingToolsExp/Private/SkinWeightsPaintTool.cpp
--- a/Engine/Plugins/Experimental/MeshModelingToolsetExp/Source/MeshModelingToolsExp/Private/SkinWeightsPaintTool.cpp
+++ b/Engine/Plugins/Experimental/MeshModelingToolsetExp/Source/MeshModelingToolsExp/Private/SkinWeightsPaintTool.cpp
@@ -1144,6 +1144,27 @@ void USkinWeightsPaintTool::ExternalUpdateWeights(const FName& BoneName, const T
 	}
 }
 
+void USkinWeightsPaintTool::HandleSkeletalMeshModified(const TArray<FName>& InBoneNames, const ESkeletalMeshNotifyType InNotifyType)
+{
+	const FName BoneName = InBoneNames.IsEmpty() ? NAME_None : InBoneNames[0];
+
+	switch (InNotifyType)
+	{
+	case ESkeletalMeshNotifyType::BonesAdded:
+		break;
+	case ESkeletalMeshNotifyType::BonesRemoved:
+		break;
+	case ESkeletalMeshNotifyType::BonesMoved:
+		break;
+	case ESkeletalMeshNotifyType::BonesSelected:
+		if (BoneName != NAME_None)
+		{
+			PendingCurrentBone = BoneName;
+		}
+		break;
+	}
+}
+
 void USkinWeightsPaintTool::OnToolPropertiesModified(UObject* ModifiedObject, FProperty* ModifiedProperty)
 {
 	// bone changed?
@@ -1151,6 +1172,10 @@ void USkinWeightsPaintTool::OnToolPropertiesModified(UObject* ModifiedObject, FP
 	if (bCurrentBoneModified)
 	{
 		PendingCurrentBone = ToolProps->CurrentBone.BoneName;
+		if (NeedsNotification())
+		{
+			GetNotifier().Notify({ToolProps->CurrentBone.BoneName}, ESkeletalMeshNotifyType::BonesSelected);
+		}
 	}
 	
 	// invalidate vertex color cache when weight color properties are modified
diff --git a/Engine/Plugins/Experimental/MeshModelingToolsetExp/Source/MeshModelingToolsExp/Public/SkeletalMesh/SkeletalMeshEditionInterface.h b/Engine/Plugins/Experimental/MeshModelingToolsetExp/Source/MeshModelingToolsExp/Public/SkeletalMesh/SkeletalMeshEditionInterface.h
new file mode 100644
--- /dev/null
+++ b/Engine/Plugins/Experimental/MeshModelingToolsetExp/Source/MeshModelingToolsExp/Public/SkeletalMesh/SkeletalMeshEditionInterface.h
@@ -0,0 +1,84 @@
+﻿// Copyright Epic Games, Inc. All Rights Reserved.
+
+#pragma once
+
+#include "UObject/Interface.h"
+#include "UObject/WeakInterfacePtr.h"
+#include "SkeletalMeshNotifier.h"
+
+#include "HitProxies.h"
+#include "GenericPlatform/ICursor.h"
+
+#include "SkeletalMeshEditionInterface.generated.h"
+
+class FSkeletalMeshToolNotifier;
+
+/**
+ * USkeletalMeshEditionInterface
+ */
+
+UINTERFACE()
+class MESHMODELINGTOOLSEXP_API USkeletalMeshEditionInterface : public UInterface
+{
+	GENERATED_BODY()
+};
+
+/**
+ * ISkeletalMeshEditionInterface
+ */
+
+class MESHMODELINGTOOLSEXP_API ISkeletalMeshEditionInterface
+{
+	GENERATED_BODY()
+
+public:
+	ISkeletalMeshNotifier& GetNotifier();
+	bool NeedsNotification() const;
+
+	void BindTo(TSharedPtr<ISkeletalMeshEditorBinding> InBinding);
+	void Unbind();
+	
+protected:
+	virtual void HandleSkeletalMeshModified(const TArray<FName>& BoneNames, const ESkeletalMeshNotifyType InNotifyType) = 0;
+	TOptional<FName> GetBoneName(HHitProxy* InHitProxy) const;
+
+private:
+	TUniquePtr<FSkeletalMeshToolNotifier> Notifier;
+	TWeakPtr<ISkeletalMeshEditorBinding> Binding;
+	
+	friend FSkeletalMeshToolNotifier;
+};
+
+/**
+ * FSkeletalMeshToolNotifier
+ */
+
+class FSkeletalMeshToolNotifier: public ISkeletalMeshNotifier
+{
+public:
+	FSkeletalMeshToolNotifier(TWeakInterfacePtr<ISkeletalMeshEditionInterface> InInterface);
+	virtual void HandleNotification(const TArray<FName>& BoneNames, const ESkeletalMeshNotifyType InNotifyType) override;
+	
+protected:
+	TWeakInterfacePtr<ISkeletalMeshEditionInterface> Interface;
+};
+
+/**
+ * HBoneHitProxy
+ */
+
+struct HBoneHitProxy : public HHitProxy
+{
+	DECLARE_HIT_PROXY()
+
+	int32 BoneIndex;
+	FName BoneName;
+
+	HBoneHitProxy(int32 InBoneIndex, FName InBoneName)
+		: HHitProxy(HPP_Foreground)
+		, BoneIndex(InBoneIndex)
+		, BoneName(InBoneName)
+	{}
+
+	virtual EMouseCursor::Type GetMouseCursor() override { return EMouseCursor::Crosshairs; }
+};
diff --git a/Engine/Plugins/Experimental/MeshModelingToolsetExp/Source/MeshModelingToolsExp/Public/SkinWeightsPaintTool.h b/Engine/Plugins/Experimental/MeshModelingToolsetExp/Source/MeshModelingToolsExp/Public/SkinWeightsPaintTool.h
--- a/Engine/Plugins/Experimental/MeshModelingToolsetExp/Source/MeshModelingToolsExp/Public/SkinWeightsPaintTool.h
+++ b/Engine/Plugins/Experimental/MeshModelingToolsetExp/Source/MeshModelingToolsExp/Public/SkinWeightsPaintTool.h
@@ -11,6 +11,7 @@
 #include "Interfaces/Interface_BoneReferenceSkeletonProvider.h"
 #include "Misc/Optional.h"
 #include "Containers/Map.h"
+#include "SkeletalMesh/SkeletalMeshEditionInterface.h"
 
 #include "SkinWeightsPaintTool.generated.h"
 
@@ -254,7 +255,7 @@ public:
  *
  */
 UCLASS()
-class MESHMODELINGTOOLSEXP_API USkinWeightsPaintTool : public UDynamicMeshBrushTool
+class MESHMODELINGTOOLSEXP_API USkinWeightsPaintTool : public UDynamicMeshBrushTool, public ISkeletalMeshEditionInterface
 {
 	GENERATED_BODY()
 
@@ -349,4 +350,6 @@ protected:
 	float MaxDrawRadius = 0.0f;
 	void UpdateBonePositionInfos(float MinRadius);
 	void RenderBonePositions(FPrimitiveDrawInterface *PDI);
+
+	virtual void HandleSkeletalMeshModified(const TArray<FName>& InBoneNames, const ESkeletalMeshNotifyType InNotifyType) override;
 };
diff --git a/Engine/Source/Developer/SkeletalMeshUtilitiesCommon/Private/SkeletalMeshNotifier.cpp b/Engine/Source/Developer/SkeletalMeshUtilitiesCommon/Private/SkeletalMeshNotifier.cpp
new file mode 100644
--- /dev/null
+++ b/Engine/Source/Developer/SkeletalMeshUtilitiesCommon/Private/SkeletalMeshNotifier.cpp
@@ -0,0 +1,14 @@
+﻿// Copyright Epic Games, Inc. All Rights Reserved.
+
+#include "SkeletalMeshNotifier.h"
+
+FSkeletalMeshNotifyDelegate& ISkeletalMeshNotifier::Delegate()
+{
+	return NotifyDelegate;
+}
+
+void ISkeletalMeshNotifier::Notify(const TArray<FName>& BoneNames, const ESkeletalMeshNotifyType InNotifyType) const
+{
+	NotifyDelegate.Broadcast(BoneNames, InNotifyType);
+}
+
diff --git a/Engine/Source/Developer/SkeletalMeshUtilitiesCommon/Public/SkeletalMeshNotifier.h b/Engine/Source/Developer/SkeletalMeshUtilitiesCommon/Public/SkeletalMeshNotifier.h
new file mode 100644
--- /dev/null
+++ b/Engine/Source/Developer/SkeletalMeshUtilitiesCommon/Public/SkeletalMeshNotifier.h
@@ -0,0 +1,50 @@
+﻿// Copyright Epic Games, Inc. All Rights Reserved.
+
+#pragma once
+
+#include "Delegates/DelegateCombinations.h"
+#include "Delegates/Delegate.h"
+#include "Containers/Array.h"
+
+class FName;
+class HHitProxy;
+
+enum class ESkeletalMeshNotifyType
+{
+	BonesAdded,
+	BonesRemoved,
+	BonesMoved,
+	BonesSelected
+};
+
+// A delegate for monitoring to skeletal mesh global notifications.
+DECLARE_MULTICAST_DELEGATE_TwoParams(FSkeletalMeshNotifyDelegate, const TArray<FName>& /*InBoneNames*/, const ESkeletalMeshNotifyType /*InNotifType*/);
+
+class SKELETALMESHUTILITIESCOMMON_API ISkeletalMeshNotifier
+{
+public:
+	ISkeletalMeshNotifier() = default;
+	virtual ~ISkeletalMeshNotifier() = default;
+
+	FSkeletalMeshNotifyDelegate& Delegate();
+
+	// override this function to react to notifications locally.
+	virtual void HandleNotification(const TArray<FName>& BoneNames, const ESkeletalMeshNotifyType InNotifyType) = 0;
+
+	void Notify(const TArray<FName>& BoneNames, const ESkeletalMeshNotifyType InNotifyType) const;
+	
+private:
+	FSkeletalMeshNotifyDelegate NotifyDelegate;
+};
+
+class SKELETALMESHUTILITIESCOMMON_API ISkeletalMeshEditorBinding
+{
+public:
+	ISkeletalMeshEditorBinding() = default;
+	virtual ~ISkeletalMeshEditorBinding() = default;
+
+	virtual ISkeletalMeshNotifier& GetNotifier() = 0;
+
+	using NameFunction = TFunction< TOptional<FName>(HHitProxy*) >;
+	virtual NameFunction GetNameFunction() = 0;
+};
\ No newline at end of file
diff --git a/Engine/Source/Editor/SkeletalMeshEditor/Private/SkeletalMeshEditor.cpp b/Engine/Source/Editor/SkeletalMeshEditor/Private/SkeletalMeshEditor.cpp
--- a/Engine/Source/Editor/SkeletalMeshEditor/Private/SkeletalMeshEditor.cpp
+++ b/Engine/Source/Editor/SkeletalMeshEditor/Private/SkeletalMeshEditor.cpp
@@ -1306,6 +1306,15 @@ void FSkeletalMeshEditor::HandleSelectionChanged(const TArrayView<TSharedPtr<ISk
 		TArray<UObject*> Objects;
 		Algo::TransformIf(InSelectedItems, Objects, [](const TSharedPtr<ISkeletonTreeItem>& InItem) { return InItem->GetObject() != nullptr; }, [](const TSharedPtr<ISkeletonTreeItem>& InItem) { return InItem->GetObject(); });
 		DetailsView->SetObjects(Objects);
+
+		if (Binding.IsValid())
+		{
+			TArray<FName> Selection;
+			Algo::TransformIf(InSelectedItems, Selection,
+				[](const TSharedPtr<ISkeletonTreeItem>& InItem) { return InItem->GetRowItemName() != NAME_None; },
+				[](const TSharedPtr<ISkeletonTreeItem>& InItem) { return InItem->GetRowItemName(); });
+			Binding->GetNotifier().Notify(Selection, ESkeletalMeshNotifyType::BonesSelected);
+		}
 	}
 }
 
@@ -1552,4 +1561,71 @@ void FSkeletalMeshEditor::HandleMeshClick(HActor* HitProxy, const FViewportClick
 	}
 }
 
+TSharedPtr<ISkeletalMeshEditorBinding> FSkeletalMeshEditor::GetBinding()
+{
+	if (!Binding)
+	{
+		Binding = MakeShared<FSkeletalMeshEditorBinding>(SharedThis(this));
+	}
+	
+	return Binding;
+}
+
+FSkeletalMeshEditorBinding::FSkeletalMeshEditorBinding(TSharedRef<FSkeletalMeshEditor> InEditor)
+	: Notifier(InEditor)
+{}
+
+ISkeletalMeshNotifier& FSkeletalMeshEditorBinding::GetNotifier()
+{
+	return Notifier;	
+}
+
+ISkeletalMeshEditorBinding::NameFunction FSkeletalMeshEditorBinding::GetNameFunction()
+{
+	return [this](HHitProxy* InHitProxy) -> TOptional<FName>
+	{
+		if (const HPersonaBoneHitProxy* BoneHitProxy = HitProxyCast<HPersonaBoneHitProxy>(InHitProxy))
+		{
+			return BoneHitProxy->BoneName;
+		}
+
+		static const TOptional<FName> Dummy;
+		return Dummy;
+	};
+}
+
+FSkeletalMeshEditorNotifier::FSkeletalMeshEditorNotifier(TSharedRef<FSkeletalMeshEditor> InEditor)
+	: ISkeletalMeshNotifier()
+	, Editor(InEditor)
+{}
+
+void FSkeletalMeshEditorNotifier::HandleNotification(const TArray<FName>& BoneNames, const ESkeletalMeshNotifyType InNotifyType)
+{
+	if (!Editor.IsValid())
+	{
+		return;
+	}
+
+	TSharedRef<ISkeletonTree> SkeletonTree = Editor.Pin()->GetSkeletonTree();
+
+	switch (InNotifyType)
+	{
+	case ESkeletalMeshNotifyType::BonesAdded:
+		SkeletonTree->Refresh();
+		break;
+	case ESkeletalMeshNotifyType::BonesRemoved:
+		SkeletonTree->Refresh();
+		break;
+	case ESkeletalMeshNotifyType::BonesMoved:
+		break;
+	case ESkeletalMeshNotifyType::BonesSelected:
+		SkeletonTree->DeselectAll();
+		if (!BoneNames.IsEmpty())
+		{
+			SkeletonTree->SetSelectedBone(BoneNames[0], ESelectInfo::Direct);
+		}
+		break;
+	}
+}
+
 #undef LOCTEXT_NAMESPACE
diff --git a/Engine/Source/Editor/SkeletalMeshEditor/Private/SkeletalMeshEditor.h b/Engine/Source/Editor/SkeletalMeshEditor/Private/SkeletalMeshEditor.h
--- a/Engine/Source/Editor/SkeletalMeshEditor/Private/SkeletalMeshEditor.h
+++ b/Engine/Source/Editor/SkeletalMeshEditor/Private/SkeletalMeshEditor.h
@@ -9,6 +9,7 @@
 #include "EditorUndoClient.h"
 #include "Toolkits/IToolkitHost.h"
 #include "ISkeletalMeshEditor.h"
+#include "SkeletalMeshNotifier.h"
 #include "Containers/ArrayView.h"
 
 class IDetailLayoutBuilder;
@@ -25,7 +26,7 @@ struct FSkeletalMeshClothBuildParams;
 struct FToolMenuContext;
 class UToolMenu;
 class SSkeletalMeshEditorToolbox;
-
+class FSkeletalMeshEditorBinding;
 
 namespace SkeletalMeshEditorModes
 {
@@ -121,6 +122,8 @@ public:
 
 	// Returns the currently hosted toolkit. Can be invalid if no toolkit is being hosted.
 	TSharedPtr<IToolkit> GetHostedToolkit() const { return HostedToolkit; }
+	
+	virtual TSharedPtr<ISkeletalMeshEditorBinding> GetBinding() override;
 
 private:
 	void HandleObjectSelected(UObject* InObject);
@@ -214,4 +217,37 @@ private:
 
 	// The toolbox widget
 	TSharedPtr<SSkeletalMeshEditorToolbox> ToolboxWidget;
+
+	// Binding to send/receive skeletal mesh modifications
+	TSharedPtr<FSkeletalMeshEditorBinding> Binding;
 };
+
+/**
+ * FSkeletalMeshEditorNotifier
+ */
+
+class FSkeletalMeshEditorNotifier: public ISkeletalMeshNotifier
+{
+public:
+	FSkeletalMeshEditorNotifier(TSharedRef<FSkeletalMeshEditor> InEditor);
+	virtual void HandleNotification(const TArray<FName>& BoneNames, const ESkeletalMeshNotifyType InNotifyType) override;
+	
+private:
+	TWeakPtr<FSkeletalMeshEditor> Editor;
+};
+
+/**
+ * FSkeletalMeshEditorBinding
+ */
+
+class FSkeletalMeshEditorBinding: public ISkeletalMeshEditorBinding
+{
+public:
+	FSkeletalMeshEditorBinding(TSharedRef<FSkeletalMeshEditor> InEditor);
+
+	virtual ISkeletalMeshNotifier& GetNotifier() override;
+	virtual NameFunction GetNameFunction() override;
+	
+private:
+	FSkeletalMeshEditorNotifier Notifier;
+};
\ No newline at end of file
diff --git a/Engine/Source/Editor/SkeletalMeshEditor/Public/ISkeletalMeshEditor.h b/Engine/Source/Editor/SkeletalMeshEditor/Public/ISkeletalMeshEditor.h
--- a/Engine/Source/Editor/SkeletalMeshEditor/Public/ISkeletalMeshEditor.h
+++ b/Engine/Source/Editor/SkeletalMeshEditor/Public/ISkeletalMeshEditor.h
@@ -6,7 +6,11 @@
 #include "PersonaAssetEditorToolkit.h"
 #include "IHasPersonaToolkit.h"
 
+class ISkeletalMeshEditorBinding;
+
 class ISkeletalMeshEditor : public FPersonaAssetEditorToolkit, public IHasPersonaToolkit
 {
+public:
+	virtual TSharedPtr<ISkeletalMeshEditorBinding> GetBinding() = 0;
 };
 
