Focus bone selection

diff --git a/Engine/Plugins/Experimental/Animation/SkeletalMeshModelingTools/Source/SkeletalMeshModelingTools/Private/SkeletalMeshModelingToolsEditorMode.cpp b/Engine/Plugins/Experimental/Animation/SkeletalMeshModelingTools/Source/SkeletalMeshModelingTools/Private/SkeletalMeshModelingToolsEditorMode.cpp
--- a/Engine/Plugins/Experimental/Animation/SkeletalMeshModelingTools/Source/SkeletalMeshModelingTools/Private/SkeletalMeshModelingToolsEditorMode.cpp
+++ b/Engine/Plugins/Experimental/Animation/SkeletalMeshModelingTools/Source/SkeletalMeshModelingTools/Private/SkeletalMeshModelingToolsEditorMode.cpp
@@ -13,6 +13,7 @@
 #include "IStylusInputModule.h"
 #include "IStylusState.h"
 #include "ToolTargets/SkeletalMeshComponentToolTarget.h"
+#include "Components/SkeletalMeshComponent.h"
 
 #include "DeformMeshPolygonsTool.h"
 #include "DisplaceMeshTool.h"
@@ -272,6 +273,52 @@ bool USkeletalMeshModelingToolsEditorMode::HandleClick(FEditorViewportClient* In
 	return Super::HandleClick(InViewportClient, HitProxy, Click);
 }
 
+bool USkeletalMeshModelingToolsEditorMode::ComputeBoundingBoxForViewportFocus(AActor* Actor, UPrimitiveComponent* PrimitiveComponent, FBox& InOutBox) const
+{
+	if (const USkeletalMeshComponent* Component = Cast<USkeletalMeshComponent>(PrimitiveComponent))
+	{
+		check(Component->GetSkeletalMeshAsset());
+		
+		if (Binding.IsValid())
+		{
+			const TArray<FName> Selection = Binding.Pin()->GetSelectedBones();
+			if (!Selection.IsEmpty())
+			{
+				TArray<FName> AllChildren;
+
+				const FReferenceSkeleton& RefSkeleton = Component->GetSkeletalMeshAsset()->GetRefSkeleton();
+				for (const FName& BoneName: Selection)
+				{
+					const int32 BoneIndex = RefSkeleton.FindBoneIndex(BoneName);
+					if (BoneIndex > INDEX_NONE)
+					{
+						// enlarge box
+						InOutBox += Component->GetBoneLocation(BoneName);
+
+						// get direct children
+						TArray<int32> Children;
+						RefSkeleton.GetDirectChildBones(BoneIndex, Children);
+						Algo::Transform(Children, AllChildren, [&RefSkeleton](int ChildrenIndex)
+						{
+							return RefSkeleton.GetBoneName(ChildrenIndex);
+						});
+					}
+				}
+
+				// enlarge box using direct children
+				for (const FName& BoneName: AllChildren)
+				{
+					InOutBox += Component->GetBoneLocation(BoneName);
+				}
+				
+				return true;
+			}
+		}
+	}
+	
+	return Super::ComputeBoundingBoxForViewportFocus(Actor, PrimitiveComponent, InOutBox);
+}
+
 void USkeletalMeshModelingToolsEditorMode::OnToolStarted(UInteractiveToolManager* Manager, UInteractiveTool* Tool)
 {
 	ConnectTool(Tool);
diff --git a/Engine/Plugins/Experimental/Animation/SkeletalMeshModelingTools/Source/SkeletalMeshModelingTools/Private/SkeletalMeshModelingToolsEditorMode.h b/Engine/Plugins/Experimental/Animation/SkeletalMeshModelingTools/Source/SkeletalMeshModelingTools/Private/SkeletalMeshModelingToolsEditorMode.h
--- a/Engine/Plugins/Experimental/Animation/SkeletalMeshModelingTools/Source/SkeletalMeshModelingTools/Private/SkeletalMeshModelingToolsEditorMode.h
+++ b/Engine/Plugins/Experimental/Animation/SkeletalMeshModelingTools/Source/SkeletalMeshModelingTools/Private/SkeletalMeshModelingToolsEditorMode.h
@@ -39,6 +39,8 @@ public:
 	// void DrawHUD(FEditorViewportClient* InViewportClient, FViewport* InViewport, const FSceneView* InView, FCanvas* InCanvas) override;
 	virtual bool HandleClick(FEditorViewportClient* InViewportClient, HHitProxy *HitProxy, const FViewportClick &Click) override;
 
+	virtual bool ComputeBoundingBoxForViewportFocus(AActor* Actor, UPrimitiveComponent* PrimitiveComponent, FBox& InOutBox) const override;
+	
 	virtual bool UsesToolkits() const override { return true; }
 
 	// binding
diff --git a/Engine/Source/Developer/SkeletalMeshUtilitiesCommon/Public/SkeletalMeshNotifier.h b/Engine/Source/Developer/SkeletalMeshUtilitiesCommon/Public/SkeletalMeshNotifier.h
--- a/Engine/Source/Developer/SkeletalMeshUtilitiesCommon/Public/SkeletalMeshNotifier.h
+++ b/Engine/Source/Developer/SkeletalMeshUtilitiesCommon/Public/SkeletalMeshNotifier.h
@@ -47,4 +47,6 @@ public:
 
 	using NameFunction = TFunction< TOptional<FName>(HHitProxy*) >;
 	virtual NameFunction GetNameFunction() = 0;
+
+	virtual TArray<FName> GetSelectedBones() const = 0;
 };
\ No newline at end of file
diff --git a/Engine/Source/Editor/SkeletalMeshEditor/Private/SkeletalMeshEditor.cpp b/Engine/Source/Editor/SkeletalMeshEditor/Private/SkeletalMeshEditor.cpp
--- a/Engine/Source/Editor/SkeletalMeshEditor/Private/SkeletalMeshEditor.cpp
+++ b/Engine/Source/Editor/SkeletalMeshEditor/Private/SkeletalMeshEditor.cpp
@@ -1572,7 +1572,8 @@ TSharedPtr<ISkeletalMeshEditorBinding> FSkeletalMeshEditor::GetBinding()
 }
 
 FSkeletalMeshEditorBinding::FSkeletalMeshEditorBinding(TSharedRef<FSkeletalMeshEditor> InEditor)
-	: Notifier(InEditor)
+	: Editor(InEditor)
+	, Notifier(InEditor)
 {}
 
 ISkeletalMeshNotifier& FSkeletalMeshEditorBinding::GetNotifier()
@@ -1594,6 +1595,24 @@ ISkeletalMeshEditorBinding::NameFunction FSkeletalMeshEditorBinding::GetNameFunc
 	};
 }
 
+TArray<FName> FSkeletalMeshEditorBinding::GetSelectedBones() const
+{
+	TArray<FName> Selection;
+
+	if (!Editor.IsValid())
+	{
+		return Selection;
+	}
+
+	const TArray<TSharedPtr<ISkeletonTreeItem>> SelectedItems = Editor.Pin()->GetSkeletonTree()->GetSelectedItems();
+
+	Algo::TransformIf(SelectedItems, Selection,
+	[](const TSharedPtr<ISkeletonTreeItem>& InItem) { return InItem->IsOfTypeByName("FSkeletonTreeBoneItem") && InItem->GetRowItemName() != NAME_None; },
+	[](const TSharedPtr<ISkeletonTreeItem>& InItem) { return InItem->GetRowItemName(); });
+	
+	return Selection;
+}
+
 FSkeletalMeshEditorNotifier::FSkeletalMeshEditorNotifier(TSharedRef<FSkeletalMeshEditor> InEditor)
 	: ISkeletalMeshNotifier()
 	, Editor(InEditor)
diff --git a/Engine/Source/Editor/SkeletalMeshEditor/Private/SkeletalMeshEditor.h b/Engine/Source/Editor/SkeletalMeshEditor/Private/SkeletalMeshEditor.h
--- a/Engine/Source/Editor/SkeletalMeshEditor/Private/SkeletalMeshEditor.h
+++ b/Engine/Source/Editor/SkeletalMeshEditor/Private/SkeletalMeshEditor.h
@@ -247,7 +247,9 @@ public:
 
 	virtual ISkeletalMeshNotifier& GetNotifier() override;
 	virtual NameFunction GetNameFunction() override;
+	virtual TArray<FName> GetSelectedBones() const override;
 	
 private:
+	TWeakPtr<FSkeletalMeshEditor> Editor;
 	FSkeletalMeshEditorNotifier Notifier;
 };
\ No newline at end of file
