Skeletal Mesh: Add the ability to create and paint vertex maps from the skeletal mesh modeling tools.

#rnx

diff --git a/Engine/Plugins/Experimental/Animation/SkeletalMeshModelingTools/Source/SkeletalMeshModelingTools/Private/SkeletalMeshModelingToolsEditorMode.cpp b/Engine/Plugins/Experimental/Animation/SkeletalMeshModelingTools/Source/SkeletalMeshModelingTools/Private/SkeletalMeshModelingToolsEditorMode.cpp
--- a/Engine/Plugins/Experimental/Animation/SkeletalMeshModelingTools/Source/SkeletalMeshModelingTools/Private/SkeletalMeshModelingToolsEditorMode.cpp
+++ b/Engine/Plugins/Experimental/Animation/SkeletalMeshModelingTools/Source/SkeletalMeshModelingTools/Private/SkeletalMeshModelingToolsEditorMode.cpp
@@ -2,6 +2,7 @@
 
 #include "SkeletalMeshModelingToolsEditorMode.h"
 
+#include "AttributeEditorTool.h"
 #include "SkeletalMeshModelingToolsEditorModeToolkit.h"
 #include "SkeletalMeshModelingToolsCommands.h"
 #include "SkeletalMeshGizmoUtils.h"
@@ -181,31 +182,6 @@ void USkeletalMeshModelingToolsEditorMode::Enter()
 
 	const FModelingToolsManagerCommands& ToolManagerCommands = FModelingToolsManagerCommands::Get();
 
-		/*
-		ToolbarBuilder.AddToolBarButton(Commands.BeginPolyEditTool);
-		ToolbarBuilder.AddToolBarButton(Commands.BeginPolyDeformTool);
-		ToolbarBuilder.AddToolBarButton(Commands.BeginHoleFillTool);
-		ToolbarBuilder.AddToolBarButton(Commands.BeginPolygonCutTool);
-	}
-	else if (PaletteName == ProcessingTabName)
-	{
-		ToolbarBuilder.AddToolBarButton(Commands.BeginSimplifyMeshTool);
-		ToolbarBuilder.AddToolBarButton(Commands.BeginRemeshMeshTool);
-		ToolbarBuilder.AddToolBarButton(Commands.BeginWeldEdgesTool);
-		ToolbarBuilder.AddToolBarButton(Commands.BeginRemoveOccludedTrianglesTool);
-		ToolbarBuilder.AddToolBarButton(Commands.BeginProjectToTargetTool);
-	}
-	else if (PaletteName == DeformTabName)
-	{
-		ToolbarBuilder.AddToolBarButton(Commands.BeginSculptMeshTool);
-		ToolbarBuilder.AddToolBarButton(Commands.BeginRemeshSculptMeshTool);
-		ToolbarBuilder.AddToolBarButton(Commands.BeginSmoothMeshTool);
-		ToolbarBuilder.AddToolBarButton(Commands.BeginOffsetMeshTool);
-		ToolbarBuilder.AddToolBarButton(Commands.BeginMeshSpaceDeformerTool);
-		ToolbarBuilder.AddToolBarButton(Commands.BeginLatticeDeformerTool);
-		ToolbarBuilder.AddToolBarButton(Commands.BeginDisplaceMeshTool);
-		*/
-
 	RegisterTool(ToolManagerCommands.BeginPolyEditTool, TEXT("BeginPolyEditTool"), NewObject<UEditMeshPolygonsToolBuilder>());
 	RegisterTool(ToolManagerCommands.BeginPolyDeformTool, TEXT("BeginPolyDeformTool"), NewObject<UDeformMeshPolygonsToolBuilder>());
 	RegisterTool(ToolManagerCommands.BeginHoleFillTool, TEXT("BeginHoleFillTool"), NewObject<UHoleFillToolBuilder>());
@@ -218,11 +194,11 @@ void USkeletalMeshModelingToolsEditorMode::Enter()
 	RegisterTool(ToolManagerCommands.BeginProjectToTargetTool, TEXT("BeginProjectToTargetTool"), NewObject<UProjectToTargetToolBuilder>());
 	
 
-	auto MoveVerticesToolBuilder = NewObject<UMeshVertexSculptToolBuilder>();
+	UMeshVertexSculptToolBuilder* MoveVerticesToolBuilder = NewObject<UMeshVertexSculptToolBuilder>();
 	MoveVerticesToolBuilder->StylusAPI = StylusStateTracker.Get();
 	RegisterTool(ToolManagerCommands.BeginSculptMeshTool, TEXT("BeginSculptMeshTool"), MoveVerticesToolBuilder);
 
-	auto DynaSculptToolBuilder = NewObject<UDynamicMeshSculptToolBuilder>();
+	UDynamicMeshSculptToolBuilder* DynaSculptToolBuilder = NewObject<UDynamicMeshSculptToolBuilder>();
 	DynaSculptToolBuilder->bEnableRemeshing = true;
 	DynaSculptToolBuilder->StylusAPI = StylusStateTracker.Get();
 	RegisterTool(ToolManagerCommands.BeginRemeshSculptMeshTool, TEXT("BeginRemeshSculptMeshTool"), DynaSculptToolBuilder);
@@ -233,7 +209,8 @@ void USkeletalMeshModelingToolsEditorMode::Enter()
 	RegisterTool(ToolManagerCommands.BeginLatticeDeformerTool, TEXT("BeginLatticeDeformerTool"), NewObject<ULatticeDeformerToolBuilder>());
 	RegisterTool(ToolManagerCommands.BeginDisplaceMeshTool, TEXT("BeginDisplaceMeshTool"), NewObject<UDisplaceMeshToolBuilder>());
 
-	// RegisterTool(ToolManagerCommands.BeginMeshAttributePaintTool, TEXT("BeginMeshAttributePaintTool"), NewObject<UMeshAttributePaintToolBuilder>());
+	RegisterTool(ToolManagerCommands.BeginAttributeEditorTool, TEXT("BeginAttributeEditorTool"), NewObject<UAttributeEditorToolBuilder>());
+	RegisterTool(ToolManagerCommands.BeginMeshAttributePaintTool, TEXT("BeginMeshAttributePaintTool"), NewObject<UMeshAttributePaintToolBuilder>());
 	RegisterTool(ToolManagerCommands.BeginSkinWeightsPaintTool, TEXT("BeginSkinWeightsPaintTool"), NewObject<USkinWeightsPaintToolBuilder>());
 	RegisterTool(ToolManagerCommands.BeginSkinWeightsBindingTool, TEXT("BeginSkinWeightsBindingTool"), NewObject<USkinWeightsBindingToolBuilder>());
 
diff --git a/Engine/Plugins/Experimental/Animation/SkeletalMeshModelingTools/Source/SkeletalMeshModelingTools/Private/SkeletalMeshModelingToolsEditorModeToolkit.cpp b/Engine/Plugins/Experimental/Animation/SkeletalMeshModelingTools/Source/SkeletalMeshModelingTools/Private/SkeletalMeshModelingToolsEditorModeToolkit.cpp
--- a/Engine/Plugins/Experimental/Animation/SkeletalMeshModelingTools/Source/SkeletalMeshModelingTools/Private/SkeletalMeshModelingToolsEditorModeToolkit.cpp
+++ b/Engine/Plugins/Experimental/Animation/SkeletalMeshModelingTools/Source/SkeletalMeshModelingTools/Private/SkeletalMeshModelingToolsEditorModeToolkit.cpp
@@ -296,8 +296,8 @@ void FSkeletalMeshModelingToolsEditorModeToolkit::BuildToolPalette(FName Palette
 	}
 	else if (PaletteName == SkinWeightsTabName)
 	{
-		// TODO: Need attribute support on skelmeshes.
-		// ToolbarBuilder.AddToolBarButton(Commands.BeginMeshAttributePaintTool);
+		ToolbarBuilder.AddToolBarButton(Commands.BeginAttributeEditorTool);
+		ToolbarBuilder.AddToolBarButton(Commands.BeginMeshAttributePaintTool);
 		ToolbarBuilder.AddToolBarButton(Commands.BeginSkinWeightsPaintTool);
 		ToolbarBuilder.AddToolBarButton(Commands.BeginSkinWeightsBindingTool);
 	}
