Skeleton Mesh Editing tool: rename plugin and module information to something more meaningful

diff --git a/Engine/Plugins/Experimental/Animation/SkeletalMeshModelingTools/SkeletalMeshModelingTools.uplugin b/Engine/Plugins/Experimental/Animation/SkeletalMeshModelingTools/SkeletalMeshModelingTools.uplugin
--- a/Engine/Plugins/Experimental/Animation/SkeletalMeshModelingTools/SkeletalMeshModelingTools.uplugin
+++ b/Engine/Plugins/Experimental/Animation/SkeletalMeshModelingTools/SkeletalMeshModelingTools.uplugin
@@ -2,8 +2,8 @@
 	"FileVersion": 3,
 	"Version": 1,
 	"VersionName": "0.1",
-	"FriendlyName": "Skeletal Mesh Modeling Tools",
-	"Description": "Modeling Tools for the Skeletal Mesh Editor",
+	"FriendlyName": "Skeletal Mesh Editing Tools",
+	"Description": "Create skeletons, paint skin weights and edit skeletal meshes.",
 	"Category" : "Animation",
 	"CreatedBy" : "Epic Games, Inc.",
 	"CreatedByURL": "https://epicgames.com/",
diff --git a/Engine/Plugins/Experimental/Animation/SkeletalMeshModelingTools/Source/SkeletalMeshModelingTools/Private/SkeletalMeshModelingToolsCommands.cpp b/Engine/Plugins/Experimental/Animation/SkeletalMeshModelingTools/Source/SkeletalMeshModelingTools/Private/SkeletalMeshModelingToolsCommands.cpp
--- a/Engine/Plugins/Experimental/Animation/SkeletalMeshModelingTools/Source/SkeletalMeshModelingTools/Private/SkeletalMeshModelingToolsCommands.cpp
+++ b/Engine/Plugins/Experimental/Animation/SkeletalMeshModelingTools/Source/SkeletalMeshModelingTools/Private/SkeletalMeshModelingToolsCommands.cpp
@@ -9,7 +9,7 @@
 
 void FSkeletalMeshModelingToolsCommands::RegisterCommands()
 {
-	UI_COMMAND(ToggleModelingToolsMode, "Enable Modeling Tools", "Toggles modeling tools on or off.", EUserInterfaceActionType::ToggleButton, FInputChord());
+	UI_COMMAND(ToggleEditingToolsMode, "Enable Editing Tools", "Toggles editing tools on or off.", EUserInterfaceActionType::ToggleButton, FInputChord());
 }
 
 const FSkeletalMeshModelingToolsCommands& FSkeletalMeshModelingToolsCommands::Get()
diff --git a/Engine/Plugins/Experimental/Animation/SkeletalMeshModelingTools/Source/SkeletalMeshModelingTools/Private/SkeletalMeshModelingToolsCommands.h b/Engine/Plugins/Experimental/Animation/SkeletalMeshModelingTools/Source/SkeletalMeshModelingTools/Private/SkeletalMeshModelingToolsCommands.h
--- a/Engine/Plugins/Experimental/Animation/SkeletalMeshModelingTools/Source/SkeletalMeshModelingTools/Private/SkeletalMeshModelingToolsCommands.h
+++ b/Engine/Plugins/Experimental/Animation/SkeletalMeshModelingTools/Source/SkeletalMeshModelingTools/Private/SkeletalMeshModelingToolsCommands.h
@@ -23,8 +23,8 @@ public:
 	virtual void RegisterCommands() override;
 	static const FSkeletalMeshModelingToolsCommands& Get();
 
-	// Modeling tools commands
-	TSharedPtr<FUICommandInfo> ToggleModelingToolsMode;
+	// Editing tools commands
+	TSharedPtr<FUICommandInfo> ToggleEditingToolsMode;
 };
 
 class FSkeletalMeshModelingToolsActionCommands : public TInteractiveToolCommands<FSkeletalMeshModelingToolsActionCommands>
diff --git a/Engine/Plugins/Experimental/Animation/SkeletalMeshModelingTools/Source/SkeletalMeshModelingTools/Private/SkeletalMeshModelingToolsModule.cpp b/Engine/Plugins/Experimental/Animation/SkeletalMeshModelingTools/Source/SkeletalMeshModelingTools/Private/SkeletalMeshModelingToolsModule.cpp
--- a/Engine/Plugins/Experimental/Animation/SkeletalMeshModelingTools/Source/SkeletalMeshModelingTools/Private/SkeletalMeshModelingToolsModule.cpp
+++ b/Engine/Plugins/Experimental/Animation/SkeletalMeshModelingTools/Source/SkeletalMeshModelingTools/Private/SkeletalMeshModelingToolsModule.cpp
@@ -75,9 +75,9 @@ void FSkeletalMeshModelingToolsModule::RegisterMenusAndToolbars()
 			if (Context && Context->SkeletalMeshEditor.IsValid())
 			{
 				InSection.AddEntry(FToolMenuEntry::InitToolBarButton(
-				    FSkeletalMeshModelingToolsCommands::Get().ToggleModelingToolsMode,
-					LOCTEXT("SkeletalMeshEditorModelingMode", "Modeling Tools"),
-				    LOCTEXT("SkeletalMeshEditorModelingModeTooltip", "Opens the Modeling Tools palette that provides selected mesh modification tools."),
+				    FSkeletalMeshModelingToolsCommands::Get().ToggleEditingToolsMode,
+					LOCTEXT("SkeletalMeshEditorModelingMode", "Editing Tools"),
+				    LOCTEXT("SkeletalMeshEditorModelingModeTooltip", "Opens the Editing Tools palette that provides selected skeletal mesh modification tools, including skeleton, skin weights and attributes editing."),
 				    FSlateIcon("ModelingToolsStyle", "LevelEditor.ModelingToolsMode")));
 			}
 		}));
@@ -114,29 +114,29 @@ TSharedRef<FExtender> FSkeletalMeshModelingToolsModule::ExtendSkelMeshEditorTool
 
 	TWeakPtr<ISkeletalMeshEditor> Ptr(InSkeletalMeshEditor);
 
-	InCommandList->MapAction(FSkeletalMeshModelingToolsCommands::Get().ToggleModelingToolsMode,
-	    FExecuteAction::CreateRaw(this, &FSkeletalMeshModelingToolsModule::OnToggleModelingToolsMode, Ptr),
+	InCommandList->MapAction(FSkeletalMeshModelingToolsCommands::Get().ToggleEditingToolsMode,
+	    FExecuteAction::CreateRaw(this, &FSkeletalMeshModelingToolsModule::OnToggleEditingToolsMode, Ptr),
 	    FCanExecuteAction(),
-	    FIsActionChecked::CreateRaw(this, &FSkeletalMeshModelingToolsModule::IsModelingToolModeActive, Ptr));
+	    FIsActionChecked::CreateRaw(this, &FSkeletalMeshModelingToolsModule::IsEditingToolModeActive, Ptr));
 
 	return ToolbarExtender.ToSharedRef();
 }
 
 
-bool FSkeletalMeshModelingToolsModule::IsModelingToolModeActive(TWeakPtr<ISkeletalMeshEditor> InSkeletalMeshEditor) const
+bool FSkeletalMeshModelingToolsModule::IsEditingToolModeActive(TWeakPtr<ISkeletalMeshEditor> InSkeletalMeshEditor) const
 {
 	TSharedPtr<ISkeletalMeshEditor> SkeletalMeshEditor = InSkeletalMeshEditor.Pin();
 	return SkeletalMeshEditor.IsValid() && SkeletalMeshEditor->GetEditorModeManager().IsModeActive(USkeletalMeshModelingToolsEditorMode::Id);
 }
 
 
-void FSkeletalMeshModelingToolsModule::OnToggleModelingToolsMode(TWeakPtr<ISkeletalMeshEditor> InSkeletalMeshEditor)
+void FSkeletalMeshModelingToolsModule::OnToggleEditingToolsMode(TWeakPtr<ISkeletalMeshEditor> InSkeletalMeshEditor)
 {
 	TSharedPtr<ISkeletalMeshEditor> SkeletalMeshEditor = InSkeletalMeshEditor.Pin();
 	if (SkeletalMeshEditor.IsValid())
 	{
 		FEditorModeTools& EditorModeManager = SkeletalMeshEditor->GetEditorModeManager();
-		if (!IsModelingToolModeActive(InSkeletalMeshEditor))
+		if (!IsEditingToolModeActive(InSkeletalMeshEditor))
 		{
 			EditorModeManager.ActivateMode(USkeletalMeshModelingToolsEditorMode::Id, true);
 
diff --git a/Engine/Plugins/Experimental/Animation/SkeletalMeshModelingTools/Source/SkeletalMeshModelingTools/Private/SkeletalMeshModelingToolsModule.h b/Engine/Plugins/Experimental/Animation/SkeletalMeshModelingTools/Source/SkeletalMeshModelingTools/Private/SkeletalMeshModelingToolsModule.h
--- a/Engine/Plugins/Experimental/Animation/SkeletalMeshModelingTools/Source/SkeletalMeshModelingTools/Private/SkeletalMeshModelingToolsModule.h
+++ b/Engine/Plugins/Experimental/Animation/SkeletalMeshModelingTools/Source/SkeletalMeshModelingTools/Private/SkeletalMeshModelingToolsModule.h
@@ -29,8 +29,8 @@ private:
 	TSharedRef<FExtender> ExtendSkelMeshEditorToolbar(const TSharedRef<FUICommandList> InCommandList, TSharedRef<ISkeletalMeshEditor> InSkeletalMeshEditor);
 
 
-	bool IsModelingToolModeActive(TWeakPtr<ISkeletalMeshEditor> InSkeletalMeshEditor) const;
-	void OnToggleModelingToolsMode(TWeakPtr<ISkeletalMeshEditor> InSkeletalMeshEditor);
+	bool IsEditingToolModeActive(TWeakPtr<ISkeletalMeshEditor> InSkeletalMeshEditor) const;
+	void OnToggleEditingToolsMode(TWeakPtr<ISkeletalMeshEditor> InSkeletalMeshEditor);
 
 	void OnPostEngineInit();
 	
