Skeletal Mesh Editing Tool: new transform widget
    - bulk edition of transforms + transaction
    - bulk edition of names
    - gizmo updates when setting transforms + orienting

diff --git a/Engine/Plugins/Experimental/Animation/SkeletalMeshModelingTools/Source/SkeletalMeshModelingTools/Private/DetailCustomization/SkeletonEditingToolPropertyCustomizations.cpp b/Engine/Plugins/Experimental/Animation/SkeletalMeshModelingTools/Source/SkeletalMeshModelingTools/Private/DetailCustomization/SkeletonEditingToolPropertyCustomizations.cpp
new file mode 100644
--- /dev/null
+++ b/Engine/Plugins/Experimental/Animation/SkeletalMeshModelingTools/Source/SkeletalMeshModelingTools/Private/DetailCustomization/SkeletonEditingToolPropertyCustomizations.cpp
@@ -0,0 +1,483 @@
+// Copyright Epic Games, Inc. All Rights Reserved.
+
+#include "SkeletonEditingToolPropertyCustomizations.h"
+
+#include "DetailLayoutBuilder.h"
+#include "Templates/UniquePtr.h"
+#include "SkeletalMesh/SkeletonEditingTool.h"
+#include "Widgets/SBoxPanel.h"
+#include "Widgets/Input/SEditableTextBox.h"
+#include "Windows/WindowsPlatformApplicationMisc.h"
+#include "InteractiveToolManager.h"
+
+
+#define LOCTEXT_NAMESPACE "SkeletonEditingToolCustomization"
+
+TSharedRef<IDetailCustomization> FSkeletonEditingPropertiesDetailCustomization::MakeInstance()
+{
+	return MakeShareable(new FSkeletonEditingPropertiesDetailCustomization);
+}
+
+void FSkeletonEditingPropertiesDetailCustomization::CustomizeDetails(IDetailLayoutBuilder& DetailBuilder)
+{
+	RelativeArray.Init(true, 4);
+	
+	TArray<TWeakObjectPtr<UObject>> ObjectsBeingCustomized;
+	DetailBuilder.GetObjectsBeingCustomized(ObjectsBeingCustomized);
+	if (ObjectsBeingCustomized.Num() > 1)
+	{
+		return;
+	}
+
+	USkeletonEditingProperties* Properties = CastChecked<USkeletonEditingProperties>(ObjectsBeingCustomized[0]);
+	Tool = Properties ? Properties->ParentTool : nullptr;
+	if (!Tool.IsValid())
+	{
+		return;
+	}
+
+	IDetailCategoryBuilder& CategoryBuilder = DetailBuilder.EditCategory(TEXT("Details"));
+
+	const TSharedRef<IPropertyHandle> NodePropHandle = DetailBuilder.GetProperty(GET_MEMBER_NAME_CHECKED(USkeletonEditingProperties, Name));
+	if (NodePropHandle->IsValidHandle())
+	{
+		TSharedRef<SWidget> ValueWidget = NodePropHandle->CreatePropertyValueWidget();
+		ValueWidget->SetEnabled(TAttribute<bool>::Create([this]{ return !Tool->GetSelection().IsEmpty(); }));
+
+		static const FText MultiValues = LOCTEXT("MultipleValues", "Multiple Values");
+		
+		DetailBuilder.AddPropertyToCategory(NodePropHandle)
+			.CustomWidget()
+			.NameContent()
+			[
+				NodePropHandle->CreatePropertyNameWidget()
+			]
+			.ValueContent()
+			[
+				SNew(SEditableTextBox)
+					.Font(FAppStyle::GetFontStyle( TEXT( "PropertyWindow.NormalFont" ) ))
+					.SelectAllTextWhenFocused(true)
+					.ClearKeyboardFocusOnCommit(false)
+					.SelectAllTextOnCommit(true)
+					.Text_Lambda([this]()
+					{
+						const TArray<FName>& Bones = Tool->GetSelection();
+						if (Bones.Num() == 0)
+						{
+							return FText::FromName(NAME_None);
+						}
+						if (Bones.Num() == 1)
+						{
+							return FText::FromName(Bones[0]);
+						}
+						return MultiValues;
+					})
+					.OnTextCommitted_Lambda([this, NodePropHandle](const FText& NewText, ETextCommit::Type) 
+					{
+						if (NewText.EqualTo(MultiValues))
+						{
+							return;
+						}
+						FText CurrentText; NodePropHandle->GetValueAsFormattedText(CurrentText);
+						if (!NewText.ToString().Equals(CurrentText.ToString(), ESearchCase::CaseSensitive))
+						{
+							NodePropHandle->SetValueFromFormattedString(NewText.ToString());
+						}
+					})
+					.IsEnabled_Lambda([this] { return !Tool->GetSelection().IsEmpty(); })
+			];
+	}
+	
+	SAdvancedTransformInputBox<FTransform>::FArguments TransformWidgetArgs = SAdvancedTransformInputBox<FTransform>::FArguments()
+		.IsEnabled(true)
+		.DisplayRelativeWorld(true)
+		.Font(IDetailLayoutBuilder::GetDetailFont())
+		.DisplayScaleLock(false)
+		.AllowEditRotationRepresentation(false);
+
+	// relative/world
+	TransformWidgetArgs
+	.OnGetIsComponentRelative_Lambda( [this](ESlateTransformComponent::Type InComponent)
+	{
+	   return RelativeArray[InComponent];
+	})
+	.OnIsComponentRelativeChanged_Lambda( [this](ESlateTransformComponent::Type InComponent, bool bIsRelative)
+	{
+		RelativeArray[InComponent] = bIsRelative;
+	});
+	
+	// get bones transforms
+	CustomizeValueGet(TransformWidgetArgs);
+		
+	// set bones transforms
+	CustomizeValueSet(TransformWidgetArgs);
+
+	// copy/paste values
+	CustomizeClipboard(TransformWidgetArgs);
+
+	// enabled
+	TransformWidgetArgs
+	.IsEnabled_Lambda([this]()
+	{
+		return !Tool->GetSelection().IsEmpty();
+	});
+	
+	SAdvancedTransformInputBox<FTransform>::ConstructGroupedTransformRows(
+		CategoryBuilder, 
+		LOCTEXT("ReferenceTransform", "Transform"),
+		LOCTEXT("ReferenceBoneTransformTooltip", "The reference transform of the bone"), 
+		TransformWidgetArgs);
+}
+
+void FSkeletonEditingPropertiesDetailCustomization::CustomizeValueGet(SAdvancedTransformInputBox<FTransform>::FArguments& InOutArgs)
+{
+	if (!Tool.IsValid())
+	{
+		return;
+	}
+	
+	auto GetNumericValue = [this](
+	   const FName InBoneName,
+	   ESlateTransformComponent::Type InComponent,
+	   ESlateRotationRepresentation::Type InRepresentation,
+	   ESlateTransformSubComponent::Type InSubComponent)
+	{
+		const bool bWorld = !RelativeArray[InComponent];
+		return SAdvancedTransformInputBox<FTransform>::GetNumericValueFromTransform(
+			  Tool->GetTransform(InBoneName, bWorld),
+			  InComponent,
+			  InRepresentation,
+			  InSubComponent);
+	};
+	
+	InOutArgs.OnGetNumericValue_Lambda( [this, GetNumericValue](
+		ESlateTransformComponent::Type InComponent,
+		ESlateRotationRepresentation::Type InRepresentation,
+		ESlateTransformSubComponent::Type InSubComponent)
+	{
+		const TArray<FName>& Bones = Tool->GetSelection();
+		if (Bones.IsEmpty())
+		{
+			return SAdvancedTransformInputBox<FTransform>::GetNumericValueFromTransform(
+				FTransform::Identity, InComponent, InRepresentation, InSubComponent);
+		}
+		
+		TOptional<FVector::FReal> Value = GetNumericValue(Bones[0], InComponent, InRepresentation, InSubComponent);
+		if (Value)
+		{
+			for (int32 Index = 1; Index < Bones.Num(); Index++)
+			{
+				const TOptional<FVector::FReal> NextValue = GetNumericValue(Bones[Index], InComponent, InRepresentation, InSubComponent);
+				if (NextValue)
+				{
+					if (!FMath::IsNearlyEqual(*Value, *NextValue))
+					{
+						return TOptional<FVector::FReal>();
+					}
+				}
+			}
+		}
+		return Value;
+	} );
+}
+
+void FSkeletonEditingPropertiesDetailCustomization::CustomizeValueSet(SAdvancedTransformInputBox<FTransform>::FArguments& InOutArgs)
+{
+	if (!Tool.IsValid())
+	{
+		return;
+	}
+	
+	auto PrepareNumericValueChanged = [this]( const FName InBoneName,
+											   ESlateTransformComponent::Type InComponent,
+											   ESlateRotationRepresentation::Type InRepresentation,
+											   ESlateTransformSubComponent::Type InSubComponent,
+											   FTransform::FReal InValue)
+	{
+		const bool bWorld = !RelativeArray[InComponent];
+		const FTransform& InTransform = Tool->GetTransform(InBoneName, bWorld);
+		FTransform OutTransform = InTransform;
+		SAdvancedTransformInputBox<FTransform>::ApplyNumericValueChange(OutTransform, InValue, InComponent, InRepresentation, InSubComponent);
+		return MakeTuple(InTransform, OutTransform);
+	};
+
+	InOutArgs.OnNumericValueChanged_Lambda([this, PrepareNumericValueChanged](
+		ESlateTransformComponent::Type InComponent,
+		ESlateRotationRepresentation::Type InRepresentation,
+		ESlateTransformSubComponent::Type InSubComponent,
+		FVector::FReal InValue)
+	{
+		const TArray<FName>& Bones = Tool->GetSelection();
+		TArray<FName> BonesToMove; BonesToMove.Reserve(Bones.Num());
+		TArray<FTransform> UpdatedTransforms; UpdatedTransforms.Reserve(Bones.Num());
+		
+		FTransform CurrentTransform, UpdatedTransform;
+		for (const FName& BoneName: Bones)
+		{
+			Tie(CurrentTransform, UpdatedTransform) =
+				PrepareNumericValueChanged(BoneName, InComponent, InRepresentation, InSubComponent, InValue);
+			if (!UpdatedTransform.Equals(CurrentTransform))
+			{
+				BonesToMove.Add(BoneName);
+				UpdatedTransforms.Add(UpdatedTransform);
+			}
+		}
+
+		if (!BonesToMove.IsEmpty())
+		{
+			if (!ActiveChange.IsValid())
+			{
+				ActiveChange = MakeUnique<SkeletonEditingTool::FRefSkeletonChange>(Tool.Get());
+			}
+		
+			const bool bWorld = !RelativeArray[InComponent];
+			Tool->SetTransforms(BonesToMove, UpdatedTransforms, bWorld);
+		}
+	})
+	.OnNumericValueCommitted_Lambda([this, PrepareNumericValueChanged](
+		ESlateTransformComponent::Type InComponent,
+		ESlateRotationRepresentation::Type InRepresentation,
+		ESlateTransformSubComponent::Type InSubComponent,
+		FVector::FReal InValue,
+		ETextCommit::Type InCommitType)
+	{
+		const TArray<FName>& Bones = Tool->GetSelection();
+
+		TArray<FName> BonesToMove; BonesToMove.Reserve(Bones.Num());
+		TArray<FTransform> UpdatedTransforms; UpdatedTransforms.Reserve(Bones.Num());
+		
+		FTransform CurrentTransform, UpdatedTransform;
+		for (const FName& BoneName: Bones)
+		{
+			Tie(CurrentTransform, UpdatedTransform) =
+				PrepareNumericValueChanged(BoneName, InComponent, InRepresentation, InSubComponent, InValue);
+			if (!UpdatedTransform.Equals(CurrentTransform))
+			{
+				BonesToMove.Add(BoneName);
+				UpdatedTransforms.Add(UpdatedTransform);
+			}
+		}
+
+		if (!BonesToMove.IsEmpty())
+		{
+			if (!ActiveChange.IsValid())
+			{
+				ActiveChange = MakeUnique<SkeletonEditingTool::FRefSkeletonChange>(Tool.Get());
+			}
+
+			const bool bWorld = !RelativeArray[InComponent];
+			Tool->SetTransforms(BonesToMove, UpdatedTransforms, bWorld);
+		}
+
+		if (ActiveChange.IsValid())
+		{
+			// send transaction
+			if (UInteractiveToolManager* ToolManager = Tool->GetToolManager())
+			{
+				ActiveChange->StoreSkeleton(Tool.Get());
+
+				static const FText TransactionDesc = LOCTEXT("ChangeNumericValue", "Change Numeric Value");
+				ToolManager->BeginUndoTransaction(TransactionDesc);
+				ToolManager->EmitObjectChange(Tool.Get(), MoveTemp(ActiveChange), TransactionDesc);
+				ToolManager->EndUndoTransaction();
+			}
+		
+			ActiveChange.Reset();
+		}
+	});
+}
+
+namespace FSkeletonEditingToolClipboardLocals
+{
+	
+template<typename DataType>
+void GetContentFromData(const DataType& InData, FString& Content)
+{
+	TBaseStructure<DataType>::Get()->ExportText(Content, &InData, &InData, nullptr, PPF_None, nullptr);
+}
+
+class FSkeletonEditingToolBoneErrorPipe : public FOutputDevice
+{
+public:
+	int32 NumErrors;
+
+	FSkeletonEditingToolBoneErrorPipe()
+		: FOutputDevice()
+		, NumErrors(0)
+	{}
+
+	virtual void Serialize(const TCHAR* V, ELogVerbosity::Type Verbosity, const class FName& Category) override
+	{
+		NumErrors++;
+	}
+};
+
+template<typename DataType>
+bool GetDataFromContent(const FString& Content, DataType& OutData)
+{
+	FSkeletonEditingToolBoneErrorPipe ErrorPipe;
+	static UScriptStruct* DataStruct = TBaseStructure<DataType>::Get();
+	DataStruct->ImportText(*Content, &OutData, nullptr, PPF_None, &ErrorPipe, DataStruct->GetName(), true);
+	return (ErrorPipe.NumErrors == 0);
+}
+	
+}
+
+void FSkeletonEditingPropertiesDetailCustomization::CustomizeClipboard(SAdvancedTransformInputBox<FTransform>::FArguments& InOutArgs)
+{
+	if (!Tool.IsValid())
+	{
+		return;
+	}
+	
+	auto OnCopyToClipboard = [this](const FName InBoneName, ESlateTransformComponent::Type InComponent)
+	{
+		using namespace FSkeletonEditingToolClipboardLocals;
+
+		const bool bWorld = !RelativeArray[InComponent];
+		const FTransform& Xfo = Tool->GetTransform(InBoneName, bWorld);
+
+		FString Content;
+		switch(InComponent)
+		{
+		case ESlateTransformComponent::Location:
+			{
+				GetContentFromData(Xfo.GetLocation(), Content);
+				break;
+			}
+		case ESlateTransformComponent::Rotation:
+			{
+				GetContentFromData(Xfo.Rotator(), Content);
+				break;
+			}
+		case ESlateTransformComponent::Scale:
+			{
+				GetContentFromData(Xfo.GetScale3D(), Content);
+				break;
+			}
+		case ESlateTransformComponent::Max:
+		default:
+			{
+				GetContentFromData(Xfo, Content);
+				break;
+			}
+		}
+
+		if (!Content.IsEmpty())
+		{
+			FPlatformApplicationMisc::ClipboardCopy(*Content);
+		}
+	};
+
+	auto PreparePasteFromClipboard = [this](const FName InBoneName, ESlateTransformComponent::Type InComponent)
+	{
+		using namespace FSkeletonEditingToolClipboardLocals;
+	
+		FString Content;
+		FPlatformApplicationMisc::ClipboardPaste(Content);
+
+		const bool bWorld = !RelativeArray[InComponent];
+		FTransform Xfo = Tool->GetTransform(InBoneName, bWorld);
+		
+		if (Content.IsEmpty())
+		{
+			return MakeTuple(Xfo, Xfo);
+		}
+
+		switch(InComponent)
+		{
+		case ESlateTransformComponent::Location:
+			{
+				FVector Data = Xfo.GetLocation();
+				if (GetDataFromContent(Content, Data))
+				{
+					Xfo.SetLocation(Data);
+				}
+				break;
+			}
+		case ESlateTransformComponent::Rotation:
+			{
+				FRotator Data = Xfo.Rotator();
+				if (GetDataFromContent(Content, Data))
+				{
+					Xfo.SetRotation(FQuat(Data));
+				}
+				break;
+			}
+		case ESlateTransformComponent::Scale:
+			{
+				FVector Data = Xfo.GetScale3D();
+				if (GetDataFromContent(Content, Data))
+				{
+					Xfo.SetScale3D(Data);
+				}
+				break;
+			}
+		case ESlateTransformComponent::Max:
+		default:
+			{
+				FTransform Data = Xfo;
+				if (GetDataFromContent(Content, Data))
+				{
+					Xfo = Data;
+				}
+				break;
+			}
+		}
+
+		return MakeTuple(Tool->GetTransform(InBoneName, bWorld), Xfo);
+	};
+
+	InOutArgs.OnCopyToClipboard_Lambda( [this, OnCopyToClipboard](ESlateTransformComponent::Type InComponent)
+	{
+		const TArray<FName>& Bones = Tool->GetSelection();
+		if (!Bones.IsEmpty())
+		{
+			OnCopyToClipboard(Bones[0], InComponent);
+		}
+	})
+	.OnPasteFromClipboard_Lambda([this, PreparePasteFromClipboard](ESlateTransformComponent::Type InComponent)
+	{
+		const TArray<FName>& Bones = Tool->GetSelection();
+		TArray<FName> BonesToMove; BonesToMove.Reserve(Bones.Num());
+		TArray<FTransform> UpdatedTransforms; UpdatedTransforms.Reserve(Bones.Num());
+		
+		FTransform CurrentTransform, UpdatedTransform;
+		for (const FName& BoneName: Bones)
+		{
+			Tie(CurrentTransform, UpdatedTransform) = PreparePasteFromClipboard(BoneName, InComponent);
+			if (!UpdatedTransform.Equals(CurrentTransform))
+			{
+				BonesToMove.Add(BoneName);
+				UpdatedTransforms.Add(UpdatedTransform);
+			}
+		}
+
+		if (BonesToMove.IsEmpty())
+		{
+			return;
+		}
+		
+		if (!ActiveChange.IsValid())
+		{
+			ActiveChange = MakeUnique<SkeletonEditingTool::FRefSkeletonChange>(Tool.Get());
+		}
+		
+		const bool bWorld = !RelativeArray[InComponent];
+		Tool->SetTransforms(BonesToMove, UpdatedTransforms, bWorld);
+
+		// send transaction
+		if (UInteractiveToolManager* ToolManager = Tool->GetToolManager())
+		{
+			ActiveChange->StoreSkeleton(Tool.Get());
+			
+			static const FText TransactionDesc = LOCTEXT("PasteTransform", "Paste Transform");
+			ToolManager->BeginUndoTransaction(TransactionDesc);
+			ToolManager->EmitObjectChange(Tool.Get(), MoveTemp(ActiveChange), TransactionDesc);
+			ToolManager->EndUndoTransaction();
+		}
+		ActiveChange.Reset();
+	});
+}
+
+#undef LOCTEXT_NAMESPACE
diff --git a/Engine/Plugins/Experimental/Animation/SkeletalMeshModelingTools/Source/SkeletalMeshModelingTools/Private/DetailCustomization/SkeletonEditingToolPropertyCustomizations.h b/Engine/Plugins/Experimental/Animation/SkeletalMeshModelingTools/Source/SkeletalMeshModelingTools/Private/DetailCustomization/SkeletonEditingToolPropertyCustomizations.h
new file mode 100644
--- /dev/null
+++ b/Engine/Plugins/Experimental/Animation/SkeletalMeshModelingTools/Source/SkeletalMeshModelingTools/Private/DetailCustomization/SkeletonEditingToolPropertyCustomizations.h
@@ -0,0 +1,28 @@
+// Copyright Epic Games, Inc. All Rights Reserved.
+
+#pragma once
+
+#include "IDetailCustomization.h"
+
+#include "SAdvancedTransformInputBox.h"
+#include "SkeletalMesh/SkeletonEditingTool.h"
+
+class IDetailLayoutBuilder;
+
+class FSkeletonEditingPropertiesDetailCustomization : public IDetailCustomization
+{
+public:
+	static TSharedRef<IDetailCustomization> MakeInstance();
+	virtual void CustomizeDetails(IDetailLayoutBuilder& DetailBuilder) override;
+	
+private:
+
+	void CustomizeValueGet(SAdvancedTransformInputBox<FTransform>::FArguments& InOutArgs);
+	void CustomizeValueSet(SAdvancedTransformInputBox<FTransform>::FArguments& InOutArgs);
+	void CustomizeClipboard(SAdvancedTransformInputBox<FTransform>::FArguments& InOutArgs);
+	
+	TWeakObjectPtr<USkeletonEditingTool> Tool;
+	TUniquePtr<SkeletonEditingTool::FRefSkeletonChange> ActiveChange;
+	TBitArray<> RelativeArray;
+};
+
diff --git a/Engine/Plugins/Experimental/Animation/SkeletalMeshModelingTools/Source/SkeletalMeshModelingTools/Private/SkeletalMeshModelingToolsModule.cpp b/Engine/Plugins/Experimental/Animation/SkeletalMeshModelingTools/Source/SkeletalMeshModelingTools/Private/SkeletalMeshModelingToolsModule.cpp
--- a/Engine/Plugins/Experimental/Animation/SkeletalMeshModelingTools/Source/SkeletalMeshModelingTools/Private/SkeletalMeshModelingToolsModule.cpp
+++ b/Engine/Plugins/Experimental/Animation/SkeletalMeshModelingTools/Source/SkeletalMeshModelingTools/Private/SkeletalMeshModelingToolsModule.cpp
@@ -15,6 +15,8 @@
 #include "Modules/ModuleManager.h"
 #include "SkeletalMeshToolMenuContext.h"
 #include "ToolMenus.h"
+#include "DetailCustomization/SkeletonEditingToolPropertyCustomizations.h"
+#include "SkeletalMesh/SkeletonEditingTool.h"
 #include "Styling/SlateIconFinder.h"
 #include "WorkflowOrientedApp/ApplicationMode.h"
 
@@ -42,10 +44,14 @@ void FSkeletalMeshModelingToolsModule::StartupModule()
 	SkelMeshEditorExtenderHandle = ToolbarExtenders.Last().GetHandle();
 
 	FCoreDelegates::OnPostEngineInit.AddRaw(this, &FSkeletalMeshModelingToolsModule::OnPostEngineInit);
+
+	RegisterPropertyCustomizations();
 }
 
 void FSkeletalMeshModelingToolsModule::ShutdownModule()
 {
+	UnregisterPropertyCustomizations();
+	
 	if (ISkeletalMeshEditorModule* SkelMeshEditorModule = FModuleManager::GetModulePtr<ISkeletalMeshEditorModule>("SkeletalMeshEditor"))
 	{
 		TArray<ISkeletalMeshEditorModule::FSkeletalMeshEditorToolbarExtender>& Extenders = SkelMeshEditorModule->GetAllSkeletalMeshEditorToolbarExtenders();
@@ -158,4 +164,47 @@ void FSkeletalMeshModelingToolsModule::OnPostEngineInit()
 	FSkeletalMeshModelingToolsActionCommands::RegisterAllToolActions();
 }
 
+void FSkeletalMeshModelingToolsModule::RegisterPropertyCustomizations()
+{
+	FPropertyEditorModule& PropertyModule = FModuleManager::LoadModuleChecked<FPropertyEditorModule>("PropertyEditor");
+
+	auto RegisterPropertyCustomization = [&](FName InStructName, auto InCustomizationFactory)
+	{
+		PropertyModule.RegisterCustomPropertyTypeLayout(
+			InStructName, 
+			FOnGetPropertyTypeCustomizationInstance::CreateStatic(InCustomizationFactory)
+			);
+		CustomizedProperties.Add(InStructName);
+	};
+
+	auto RegisterDetailCustomization = [&](FName InStructName, auto InCustomizationFactory)
+	{
+		PropertyModule.RegisterCustomClassLayout(
+			InStructName,
+			FOnGetDetailCustomizationInstance::CreateStatic(InCustomizationFactory)
+		);
+		CustomizedClasses.Add(InStructName);
+	};
+
+	RegisterDetailCustomization(USkeletonEditingProperties::StaticClass()->GetFName(), &FSkeletonEditingPropertiesDetailCustomization::MakeInstance);
+
+	PropertyModule.NotifyCustomizationModuleChanged();
+}
+
+void FSkeletalMeshModelingToolsModule::UnregisterPropertyCustomizations()
+{
+	if (FPropertyEditorModule* PropertyModule = FModuleManager::GetModulePtr<FPropertyEditorModule>("PropertyEditor"))
+	{
+		for (const FName& PropertyName: CustomizedProperties)
+		{
+			PropertyModule->UnregisterCustomPropertyTypeLayout(PropertyName);
+		}
+		for (const FName& ClassName : CustomizedClasses)
+		{
+			PropertyModule->UnregisterCustomClassLayout(ClassName);
+		}
+		PropertyModule->NotifyCustomizationModuleChanged();
+	}
+}
+
 #undef LOCTEXT_NAMESPACE
diff --git a/Engine/Plugins/Experimental/Animation/SkeletalMeshModelingTools/Source/SkeletalMeshModelingTools/Private/SkeletalMeshModelingToolsModule.h b/Engine/Plugins/Experimental/Animation/SkeletalMeshModelingTools/Source/SkeletalMeshModelingTools/Private/SkeletalMeshModelingToolsModule.h
--- a/Engine/Plugins/Experimental/Animation/SkeletalMeshModelingTools/Source/SkeletalMeshModelingTools/Private/SkeletalMeshModelingToolsModule.h
+++ b/Engine/Plugins/Experimental/Animation/SkeletalMeshModelingTools/Source/SkeletalMeshModelingTools/Private/SkeletalMeshModelingToolsModule.h
@@ -34,6 +34,12 @@ private:
 
 	void OnPostEngineInit();
 	
+	void RegisterPropertyCustomizations();
+	void UnregisterPropertyCustomizations();
+
+	TArray<FName> CustomizedProperties;
+	TArray<FName> CustomizedClasses;
+	
 	// The handle for the extender delegate we added. Needed for clean module shutdown.
 	FDelegateHandle SkelMeshEditorExtenderHandle;
 
diff --git a/Engine/Plugins/Experimental/Animation/SkeletalMeshModelingTools/Source/SkeletalMeshModelingTools/SkeletalMeshModelingTools.Build.cs b/Engine/Plugins/Experimental/Animation/SkeletalMeshModelingTools/Source/SkeletalMeshModelingTools/SkeletalMeshModelingTools.Build.cs
--- a/Engine/Plugins/Experimental/Animation/SkeletalMeshModelingTools/Source/SkeletalMeshModelingTools/SkeletalMeshModelingTools.Build.cs
+++ b/Engine/Plugins/Experimental/Animation/SkeletalMeshModelingTools/Source/SkeletalMeshModelingTools/SkeletalMeshModelingTools.Build.cs
@@ -30,7 +30,11 @@ public class SkeletalMeshModelingTools : ModuleRules
 				"ToolMenus",
 				"ToolWidgets",
 				"UnrealEd",
-				"StatusBar"
+				"StatusBar",
+				"PropertyEditor",
+				"AnimationCore",
+				"AnimationWidgets",
+				"ApplicationCore"
 			}
 		);
 	}
diff --git a/Engine/Plugins/Experimental/MeshModelingToolsetExp/Source/MeshModelingToolsExp/Private/SkeletalMesh/SkeletonEditingTool.cpp b/Engine/Plugins/Experimental/MeshModelingToolsetExp/Source/MeshModelingToolsExp/Private/SkeletalMesh/SkeletonEditingTool.cpp
--- a/Engine/Plugins/Experimental/MeshModelingToolsetExp/Source/MeshModelingToolsExp/Private/SkeletalMesh/SkeletonEditingTool.cpp
+++ b/Engine/Plugins/Experimental/MeshModelingToolsetExp/Source/MeshModelingToolsExp/Private/SkeletalMesh/SkeletonEditingTool.cpp
@@ -276,9 +276,9 @@ void USkeletonEditingTool::UpdateGizmo() const
 				if (bWorld)
 				{
 					const FTransform NewLocal = InNewTransform.GetRelativeTransform(ParentGlobal);
-					return Modifier->SetBoneTransform(BoneName, NewLocal, true);
+					return Modifier->SetBoneTransform(BoneName, NewLocal, Properties->bUpdateChildren);
 				}
-				return Modifier->SetBoneTransform(BoneName, InNewTransform, true);
+				return Modifier->SetBoneTransform(BoneName, InNewTransform, Properties->bUpdateChildren);
 			});
 	}
 }
@@ -537,14 +537,14 @@ void USkeletonEditingTool::MoveBones()
 
 void USkeletonEditingTool::RenameBones()
 {
-	const FName CurrentBone = GetCurrentBone();
-	if (CurrentBone == Properties->Name || Properties->Name == NAME_None)
+	const TArray<int32> BoneIndices = GetSelectedBoneIndexes();
+	if (BoneIndices.IsEmpty() || Properties->Name == NAME_None)
 	{
 		return;
 	}
 	
-	const FReferenceSkeleton& RefSkeleton = Modifier->GetReferenceSkeleton();
-	if (RefSkeleton.FindRawBoneIndex(CurrentBone) == INDEX_NONE)
+	const FName CurrentBone = GetCurrentBone();
+	if (BoneIndices.Num() == 1 && CurrentBone == Properties->Name)
 	{
 		return;
 	}
@@ -552,11 +552,26 @@ void USkeletonEditingTool::RenameBones()
 	TGuardValue OperationGuard(Operation, EEditingOperation::Rename);
 	BeginChange();
 
-	const bool bBoneRenamed = Modifier->RenameBone(CurrentBone, Properties->Name);
+	TArray<FName> ReversedSelection = Selection;
+	Algo::Reverse(ReversedSelection);
+
+	TArray<FName> NewNames;
+	NewNames.Init(Properties->Name, Selection.Num());
+
+	const bool bBoneRenamed = Modifier->RenameBones(ReversedSelection, NewNames);
 	if (bBoneRenamed)
 	{
-		Selection.Last() = Properties->Name;
-
+		// update selection with new names
+		const TArray<FMeshBoneInfo>& BoneInfos = Modifier->GetReferenceSkeleton().GetRawRefBoneInfo();
+		for (int32 Index = 0; Index < BoneIndices.Num(); Index++)
+		{
+			if (BoneIndices[Index] != INDEX_NONE)
+			{
+				Selection[Index] = BoneInfos[BoneIndices[Index]].Name;
+			}
+		}
+		Properties->Name = GetCurrentBone();
+		
 		if (NeedsNotification())
 		{
 			GetNotifier().Notify({CurrentBone}, ESkeletalMeshNotifyType::BonesRenamed);
@@ -636,6 +651,7 @@ void USkeletonEditingTool::OrientBones()
 	const bool bBoneOriented = Modifier->OrientBones(Bones, OrientingProperties->Options);
 	if (bBoneOriented)
 	{
+		UpdateGizmo();
 		// if (NeedsNotification())
 		// {
 		// 	GetNotifier().Notify(Bones, ESkeletalMeshNotifyType::BonesMoved);
@@ -1015,10 +1031,62 @@ FBox USkeletonEditingTool::GetWorldSpaceFocusBox()
 }
 
 TArray<FName> USkeletonEditingTool::GetSelectedBones() const
+{
+	return GetSelection();
+}
+
+const TArray<FName>& USkeletonEditingTool::GetSelection() const
 {
 	return Selection;
 }
 
+const FTransform& USkeletonEditingTool::GetTransform(const FName InBoneName, const bool bWorld) const
+{
+	if (Modifier)
+	{
+		const FReferenceSkeleton& RefSkeleton = Modifier->GetReferenceSkeleton();
+		return Modifier->GetTransform(RefSkeleton.FindRawBoneIndex(InBoneName), bWorld);
+	}
+	return FTransform::Identity;
+}
+
+void USkeletonEditingTool::SetTransforms(const TArray<FName>& InBones, const TArray<FTransform>& InTransforms, const bool bWorld) const
+{
+	const int32 NumBonesToMove = InBones.Num();
+	if (NumBonesToMove == 0 || NumBonesToMove != InTransforms.Num())
+	{
+		return;
+	}
+	
+	TArray<FTransform> NewTransforms = InTransforms;
+	if (bWorld)
+	{ // switch to local
+		const FReferenceSkeleton& RefSkeleton = Modifier->GetReferenceSkeleton();
+		const TArray<FMeshBoneInfo>& BoneInfos = RefSkeleton.GetRawRefBoneInfo();
+		
+		for (int32 Index = 0; Index < InBones.Num(); ++Index)
+		{
+			const int32 BoneIndex = RefSkeleton.FindRawBoneIndex(InBones[Index]);
+			check(BoneIndex != INDEX_NONE);
+
+			const int32 ParentBoneIndex = BoneInfos[BoneIndex].ParentIndex;
+			const int32 ParentToBeSet = ParentBoneIndex != INDEX_NONE ?
+				InBones.IndexOfByKey(BoneInfos[ParentBoneIndex].Name) : INDEX_NONE;
+
+			const FTransform& ParentGlobal = ParentToBeSet != INDEX_NONE ?
+				InTransforms[ParentToBeSet] : Modifier->GetTransform(ParentBoneIndex, true);
+			
+			NewTransforms[Index] = NewTransforms[Index].GetRelativeTransform(ParentGlobal);
+		}
+	}
+
+	const bool bBonesMoved = Modifier->SetBonesTransforms(InBones, NewTransforms, Properties->bUpdateChildren);
+	if (bBonesMoved)
+	{
+		UpdateGizmo();
+	}
+}
+
 FName USkeletonEditingTool::GetCurrentBone() const
 {
 	return Selection.IsEmpty() ? NAME_None : Selection.Last(); 
@@ -1094,11 +1162,6 @@ void USkeletonEditingProperties::PostEditChangeProperty(FPropertyChangedEvent& P
 		{
 			ParentTool->RenameBones();
 		}
-
-		if (PropertyName == GET_MEMBER_NAME_CHECKED(USkeletonEditingProperties, Transform))
-		{
-			ParentTool->MoveBones();
-		}
 	}
 }
 #endif
diff --git a/Engine/Plugins/Experimental/MeshModelingToolsetExp/Source/MeshModelingToolsExp/Public/SkeletalMesh/SkeletonEditingTool.h b/Engine/Plugins/Experimental/MeshModelingToolsetExp/Source/MeshModelingToolsExp/Public/SkeletalMesh/SkeletonEditingTool.h
--- a/Engine/Plugins/Experimental/MeshModelingToolsetExp/Source/MeshModelingToolsExp/Public/SkeletalMesh/SkeletonEditingTool.h
+++ b/Engine/Plugins/Experimental/MeshModelingToolsetExp/Source/MeshModelingToolsExp/Public/SkeletalMesh/SkeletonEditingTool.h
@@ -28,7 +28,7 @@ namespace SkeletonEditingTool
 /**
  * A wrapper change class that stores a reference skeleton and the bones' indexes trackers to be used for undo/redo.
  */
-class FRefSkeletonChange : public FToolCommandChange
+class MESHMODELINGTOOLSEXP_API FRefSkeletonChange : public FToolCommandChange
 {
 public:
 	FRefSkeletonChange(const USkeletonEditingTool* InTool);
@@ -136,6 +136,11 @@ public:
 	// ISkeletalMeshEditionInterface overrides
 	virtual TArray<FName> GetSelectedBones() const override;
 
+	const TArray<FName>& GetSelection() const;
+
+	const FTransform& GetTransform(const FName InBoneName, const bool bWorld) const;
+	void SetTransforms(const TArray<FName>& InBones, const TArray<FTransform>& InTransforms, const bool bWorld) const;
+
 	// IModifierToggleBehaviorTarget overrides
 	virtual void OnUpdateModifierState(int ModifierID, bool bIsOn) override;
 
@@ -248,7 +253,7 @@ public:
 	UPROPERTY(EditAnywhere, Category = "Details")
 	FName Name;
 
-	UPROPERTY(EditAnywhere, Category = "Details")
+	UPROPERTY()
 	FTransform Transform;
 	
 	UPROPERTY(EditAnywhere, Category = "Move")
@@ -264,7 +269,6 @@ public:
 	virtual void PostEditChangeProperty(FPropertyChangedEvent &PropertyChangedEvent) override;
 #endif
 	
-private:
 	TWeakObjectPtr<USkeletonEditingTool> ParentTool = nullptr;
 };
 
diff --git a/Engine/Plugins/Experimental/MeshModelingToolsetExp/Source/SkeletalMeshModifiers/Private/SkeletonModifier.cpp b/Engine/Plugins/Experimental/MeshModelingToolsetExp/Source/SkeletalMeshModifiers/Private/SkeletonModifier.cpp
--- a/Engine/Plugins/Experimental/MeshModelingToolsetExp/Source/SkeletalMeshModifiers/Private/SkeletonModifier.cpp
+++ b/Engine/Plugins/Experimental/MeshModelingToolsetExp/Source/SkeletalMeshModifiers/Private/SkeletonModifier.cpp
@@ -4,7 +4,6 @@
 #include "SkeletonModifier.h"
 
 #include "BoneWeights.h"
-#include "LODUtilities.h"
 #include "MeshDescription.h"
 #include "Animation/Skeleton.h"
 #include "ReferenceSkeleton.h"
@@ -14,7 +13,8 @@
 #include "Animation/AnimMontage.h"
 #include "AssetRegistry/AssetRegistryModule.h"
 #include "AssetRegistry/AssetData.h"
-#include "Rendering/SkeletalMeshModel.h"
+#include "Engine/SkeletalMesh.h"
+#include "RenderingThread.h"
 
 
 namespace USkeletonModifierLocals
@@ -979,7 +979,7 @@ bool USkeletonModifier::RenameBones(const TArray<FName>& InOldBoneNames, const T
 		for (int32 Index = 0; Index < NumBonesToRename; ++Index)
 		{
 			const FName& OldName = InOldBoneNames[Index];
-			const FName& NewName = InNewBoneNames[Index];
+			const FName NewName = GetUniqueName(InNewBoneNames[Index]);
 			if (OldName != NAME_None && OldName != NewName)
 			{
 				Modifier.Rename(OldName, NewName);
@@ -1219,11 +1219,47 @@ void USkeletonModifier::UpdateBoneTracker(const TArray<FMeshBoneInfo>& InOtherIn
 
 FName USkeletonModifier::GetUniqueName(const FName InBoneName) const
 {
-	if (!ReferenceSkeleton)
+	if (!ReferenceSkeleton || InBoneName == NAME_None)
 	{
 		return NAME_None;
 	}
-	
+
+	// remove unwanted characters
+	auto GetSanitizedName = [](const FName InName)
+	{
+		bool bHasAnyGoodChar = false;
+		FString Name = InName.ToString(); 
+		Name.TrimStartAndEndInline();
+		for (int32 i = 0; i < Name.Len(); ++i)
+		{
+			TCHAR& C = Name[i];
+			const bool bGoodChar = FChar::IsAlpha(C) || (C == '_') || ((i > 0) && FChar::IsDigit(C));
+			bHasAnyGoodChar |= bGoodChar;
+			if (!bGoodChar)
+			{
+				C = '_';
+			}
+		}
+		return bHasAnyGoodChar ? FName(*Name) : NAME_None;
+	};
+
+	// remove digits at the end
+	auto GetPrefix = [](const FName InName) -> FName
+	{
+		const FString Name = InName.ToString();
+		const int32 LastNonDigit = Name.FindLastCharByPredicate([](TCHAR C) { return !FChar::IsDigit(C); });
+		if (LastNonDigit == INDEX_NONE)
+		{
+			return NAME_None;
+		}
+		if (LastNonDigit == Name.Len()-1)
+		{
+			return InName;
+		}
+		return FName(*Name.Left(LastNonDigit+1));
+	};
+
+	// check for availability
 	auto IsNameAvailable = [&](const FName Name)
 	{
 		const int32 Index = ReferenceSkeleton->FindRawBoneIndex(Name);
@@ -1234,12 +1270,23 @@ FName USkeletonModifier::GetUniqueName(const FName InBoneName) const
 		return true;
 	};
 
-	FName OutName = InBoneName;
-	int32 CurrentIndex = 0;
+	const FName SanitizedName = GetSanitizedName(InBoneName);
+	if (SanitizedName == NAME_None)
+	{
+		return NAME_None;
+	}
+	
+	const FName Prefix = GetPrefix(SanitizedName);
+	if (Prefix == NAME_None)
+	{
+		return NAME_None;
+	}
 
+	FName OutName = Prefix;
+	int32 CurrentIndex = 0;
 	while (!IsNameAvailable(OutName))
 	{
-		OutName = FName(*FString::Printf(TEXT("%s%d"), *InBoneName.ToString(), CurrentIndex++));
+		OutName = FName(*FString::Printf(TEXT("%s%d"), *Prefix.ToString(), CurrentIndex++));
 	}
 
 	return OutName;
diff --git a/Engine/Plugins/Experimental/MeshModelingToolsetExp/Source/SkeletalMeshModifiers/Public/SkeletonModifier.h b/Engine/Plugins/Experimental/MeshModelingToolsetExp/Source/SkeletalMeshModifiers/Public/SkeletonModifier.h
--- a/Engine/Plugins/Experimental/MeshModelingToolsetExp/Source/SkeletalMeshModifiers/Public/SkeletonModifier.h
+++ b/Engine/Plugins/Experimental/MeshModelingToolsetExp/Source/SkeletalMeshModifiers/Public/SkeletonModifier.h
@@ -206,6 +206,7 @@ private:
 	TUniquePtr<FReferenceSkeleton> ReferenceSkeleton;
 	TUniquePtr<FTransformComposer> TransformComposer;
 	TArray<int32> BoneIndexTracker;
+	bool bRestricted = false;
 };
 
 struct FTransformComposer
