Skeletal Editing: Remove unnecessary plugin platform restriction.

#rnx

diff --git a/Engine/Plugins/Experimental/Animation/SkeletalMeshModelingTools/SkeletalMeshModelingTools.uplugin b/Engine/Plugins/Experimental/Animation/SkeletalMeshModelingTools/SkeletalMeshModelingTools.uplugin
--- a/Engine/Plugins/Experimental/Animation/SkeletalMeshModelingTools/SkeletalMeshModelingTools.uplugin
+++ b/Engine/Plugins/Experimental/Animation/SkeletalMeshModelingTools/SkeletalMeshModelingTools.uplugin
@@ -17,8 +17,7 @@
 		{
 			"Name": "SkeletalMeshModelingTools",
 			"Type": "Editor",
-			"LoadingPhase": "Default",
-			"PlatformAllowList": [ "Win64" ]
+			"LoadingPhase": "Default"
 		}
 	],
 	"Plugins": [
diff --git a/Engine/Plugins/Experimental/Animation/SkeletalMeshModelingTools/Source/SkeletalMeshModelingTools/Private/DetailCustomization/SkeletonEditingToolPropertyCustomizations.cpp b/Engine/Plugins/Experimental/Animation/SkeletalMeshModelingTools/Source/SkeletalMeshModelingTools/Private/DetailCustomization/SkeletonEditingToolPropertyCustomizations.cpp
--- a/Engine/Plugins/Experimental/Animation/SkeletalMeshModelingTools/Source/SkeletalMeshModelingTools/Private/DetailCustomization/SkeletonEditingToolPropertyCustomizations.cpp
+++ b/Engine/Plugins/Experimental/Animation/SkeletalMeshModelingTools/Source/SkeletalMeshModelingTools/Private/DetailCustomization/SkeletonEditingToolPropertyCustomizations.cpp
@@ -7,7 +7,7 @@
 #include "SkeletalMesh/SkeletonEditingTool.h"
 #include "Widgets/SBoxPanel.h"
 #include "Widgets/Input/SEditableTextBox.h"
-#include "Windows/WindowsPlatformApplicationMisc.h"
+#include "HAL/PlatformApplicationMisc.h"
 #include "InteractiveToolManager.h"
 
 
