make sure the modeling commands are registered

diff --git a/Engine/Plugins/Experimental/Animation/SkeletalMeshModelingTools/Source/SkeletalMeshModelingTools/Private/SkeletalMeshModelingToolsModule.cpp b/Engine/Plugins/Experimental/Animation/SkeletalMeshModelingTools/Source/SkeletalMeshModelingTools/Private/SkeletalMeshModelingToolsModule.cpp
--- a/Engine/Plugins/Experimental/Animation/SkeletalMeshModelingTools/Source/SkeletalMeshModelingTools/Private/SkeletalMeshModelingToolsModule.cpp
+++ b/Engine/Plugins/Experimental/Animation/SkeletalMeshModelingTools/Source/SkeletalMeshModelingTools/Private/SkeletalMeshModelingToolsModule.cpp
@@ -12,6 +12,7 @@
 #include "Engine/SkeletalMesh.h"
 #include "ISkeletalMeshEditor.h"
 #include "ISkeletalMeshEditorModule.h"
+#include "ModelingToolsManagerActions.h"
 #include "Modules/ModuleManager.h"
 #include "SkeletalMeshToolMenuContext.h"
 #include "ToolMenus.h"
@@ -63,6 +64,7 @@ void FSkeletalMeshModelingToolsModule::ShutdownModule()
 
 	FSkeletalMeshModelingToolsCommands::Unregister();
 	FSkeletalMeshModelingToolsStyle::Unregister();
+	FModelingToolsManagerCommands::Unregister();
 	
 	FCoreDelegates::OnPostEngineInit.RemoveAll(this);
 	FSkeletalMeshModelingToolsActionCommands::UnregisterAllToolActions();
@@ -162,6 +164,7 @@ void FSkeletalMeshModelingToolsModule::OnToggleEditingToolsMode(TWeakPtr<ISkelet
 void FSkeletalMeshModelingToolsModule::OnPostEngineInit()
 {
 	FSkeletalMeshModelingToolsActionCommands::RegisterAllToolActions();
+	FModelingToolsManagerCommands::Register();
 }
 
 void FSkeletalMeshModelingToolsModule::RegisterPropertyCustomizations()
