remove stylus plugin from skeletal modeling mode and comment out related code, so it could in theory be brought back in the future

diff --git a/Engine/Plugins/Experimental/Animation/SkeletalMeshModelingTools/Source/SkeletalMeshModelingTools/Private/SkeletalMeshModelingToolsEditorMode.cpp b/Engine/Plugins/Experimental/Animation/SkeletalMeshModelingTools/Source/SkeletalMeshModelingTools/Private/SkeletalMeshModelingToolsEditorMode.cpp
--- a/Engine/Plugins/Experimental/Animation/SkeletalMeshModelingTools/Source/SkeletalMeshModelingTools/Private/SkeletalMeshModelingToolsEditorMode.cpp
+++ b/Engine/Plugins/Experimental/Animation/SkeletalMeshModelingTools/Source/SkeletalMeshModelingTools/Private/SkeletalMeshModelingToolsEditorMode.cpp
@@ -12,8 +12,17 @@
 #include "EdModeInteractiveToolsContext.h"
 #include "EditorModeManager.h"
 #include "EditorViewportClient.h"
+#include "ToolTargets/SkeletalMeshComponentToolTarget.h"
+
+// Stylus support is currently disabled due to issues with the stylus plugin
+// We are leaving the code in this cpp file, defined out, so that it is easier to bring back if/when the stylus plugin is improved.
+#define ENABLE_STYLUS_SUPPORT 0
+
+#if ENABLE_STYLUS_SUPPORT 
 #include "IStylusInputModule.h"
 #include "IStylusState.h"
+#endif
+
 #include "ToolTargets/SkeletalMeshComponentToolTarget.h"
 #include "Components/SkeletalMeshComponent.h"
 
@@ -48,7 +57,7 @@
 
 #define LOCTEXT_NAMESPACE "SkeletalMeshModelingToolsEditorMode"
 
-
+#if ENABLE_STYLUS_SUPPORT
 // FStylusStateTracker registers itself as a listener for stylus events and implements
 // the IToolStylusStateProviderAPI interface, which allows MeshSurfacePointTool implementations
  // to query for the pen pressure.
@@ -134,6 +143,7 @@ public:
 	}
 
 };
+#endif // ENABLE_STYLUS_SUPPORT
 
 
 // NOTE: This is a simple proxy at the moment. In the future we want to pull in more of the 
@@ -175,7 +185,9 @@ void USkeletalMeshModelingToolsEditorMode::Enter()
 	
 	InteractiveToolsContext->TargetManager->AddTargetFactory(NewObject<USkeletalMeshComponentToolTargetFactory>(InteractiveToolsContext->TargetManager));
 
+#if ENABLE_STYLUS_SUPPORT
 	StylusStateTracker = MakeUnique<FStylusStateTracker>();
+#endif
 	// register gizmo helper
 	UE::TransformGizmoUtil::RegisterTransformGizmoContextObject(InteractiveToolsContext);
 	UE::SkeletalMeshGizmoUtils::RegisterTransformGizmoContextObject(InteractiveToolsContext);
@@ -195,12 +207,16 @@ void USkeletalMeshModelingToolsEditorMode::Enter()
 	
 
 	UMeshVertexSculptToolBuilder* MoveVerticesToolBuilder = NewObject<UMeshVertexSculptToolBuilder>();
+#if ENABLE_STYLUS_SUPPORT
 	MoveVerticesToolBuilder->StylusAPI = StylusStateTracker.Get();
+#endif
 	RegisterTool(ToolManagerCommands.BeginSculptMeshTool, TEXT("BeginSculptMeshTool"), MoveVerticesToolBuilder);
 
 	UDynamicMeshSculptToolBuilder* DynaSculptToolBuilder = NewObject<UDynamicMeshSculptToolBuilder>();
 	DynaSculptToolBuilder->bEnableRemeshing = true;
+#if ENABLE_STYLUS_SUPPORT
 	DynaSculptToolBuilder->StylusAPI = StylusStateTracker.Get();
+#endif
 	RegisterTool(ToolManagerCommands.BeginRemeshSculptMeshTool, TEXT("BeginRemeshSculptMeshTool"), DynaSculptToolBuilder);
 	
 	RegisterTool(ToolManagerCommands.BeginSmoothMeshTool, TEXT("BeginSmoothMeshTool"), NewObject<USmoothMeshToolBuilder>());
@@ -236,8 +252,9 @@ void USkeletalMeshModelingToolsEditorMode::Exit()
 	UEditorInteractiveToolsContext* InteractiveToolsContext = GetInteractiveToolsContext();
 	UE::TransformGizmoUtil::DeregisterTransformGizmoContextObject(InteractiveToolsContext);
 	UE::SkeletalMeshGizmoUtils::UnregisterTransformGizmoContextObject(InteractiveToolsContext);
-	
+#if ENABLE_STYLUS_SUPPORT
 	StylusStateTracker = nullptr;
+#endif
 
 	UEdMode::Exit();
 }
diff --git a/Engine/Plugins/Experimental/Animation/SkeletalMeshModelingTools/Source/SkeletalMeshModelingTools/Private/SkeletalMeshModelingToolsEditorMode.h b/Engine/Plugins/Experimental/Animation/SkeletalMeshModelingTools/Source/SkeletalMeshModelingTools/Private/SkeletalMeshModelingToolsEditorMode.h
--- a/Engine/Plugins/Experimental/Animation/SkeletalMeshModelingTools/Source/SkeletalMeshModelingTools/Private/SkeletalMeshModelingToolsEditorMode.h
+++ b/Engine/Plugins/Experimental/Animation/SkeletalMeshModelingTools/Source/SkeletalMeshModelingTools/Private/SkeletalMeshModelingToolsEditorMode.h
@@ -7,7 +7,7 @@
 
 #include "SkeletalMeshModelingToolsEditorMode.generated.h"
 
-class FStylusStateTracker;
+//class FStylusStateTracker;
 class FSkeletalMeshModelingToolsEditorModeToolkit;
 class UEdModeInteractiveToolsContext;
 class ISkeletalMeshNotifier;
@@ -52,7 +52,8 @@ protected:
 	virtual void OnToolEnded(UInteractiveToolManager* Manager, UInteractiveTool* Tool) override;
 	
 private:
-	TUniquePtr<FStylusStateTracker> StylusStateTracker;
+	// Stylus support is currently disabled; this is left in for reference if/when it is brought back
+	//TUniquePtr<FStylusStateTracker> StylusStateTracker;
 
 	static ISkeletalMeshEditionInterface* GetSkeletonInterface(UInteractiveTool* InTool);
 
diff --git a/Engine/Plugins/Experimental/Animation/SkeletalMeshModelingTools/Source/SkeletalMeshModelingTools/SkeletalMeshModelingTools.Build.cs b/Engine/Plugins/Experimental/Animation/SkeletalMeshModelingTools/Source/SkeletalMeshModelingTools/SkeletalMeshModelingTools.Build.cs
--- a/Engine/Plugins/Experimental/Animation/SkeletalMeshModelingTools/Source/SkeletalMeshModelingTools/SkeletalMeshModelingTools.Build.cs
+++ b/Engine/Plugins/Experimental/Animation/SkeletalMeshModelingTools/Source/SkeletalMeshModelingTools/SkeletalMeshModelingTools.Build.cs
@@ -26,7 +26,6 @@ public class SkeletalMeshModelingTools : ModuleRules
 				"SkeletalMeshUtilitiesCommon",
 				"Slate",
 				"SlateCore",
-				"StylusInput",
 				"ToolMenus",
 				"ToolWidgets",
 				"UnrealEd",
