remove extra copy of an include accidentally added by a merge from SkeletalMeshModelingToolsEditorMode.cpp

diff --git a/Engine/Plugins/Experimental/Animation/SkeletalMeshModelingTools/Source/SkeletalMeshModelingTools/Private/SkeletalMeshModelingToolsEditorMode.cpp b/Engine/Plugins/Experimental/Animation/SkeletalMeshModelingTools/Source/SkeletalMeshModelingTools/Private/SkeletalMeshModelingToolsEditorMode.cpp
--- a/Engine/Plugins/Experimental/Animation/SkeletalMeshModelingTools/Source/SkeletalMeshModelingTools/Private/SkeletalMeshModelingToolsEditorMode.cpp
+++ b/Engine/Plugins/Experimental/Animation/SkeletalMeshModelingTools/Source/SkeletalMeshModelingTools/Private/SkeletalMeshModelingToolsEditorMode.cpp
@@ -12,7 +12,6 @@
 #include "EdModeInteractiveToolsContext.h"
 #include "EditorModeManager.h"
 #include "EditorViewportClient.h"
-#include "ToolTargets/SkeletalMeshComponentToolTarget.h"
 
 // Stylus support is currently disabled due to issues with the stylus plugin
 // We are leaving the code in this cpp file, defined out, so that it is easier to bring back if/when the stylus plugin is improved.
