register FModelingToolsEditorModeStyle before FModelingToolsManagerCommands

diff --git a/Engine/Plugins/Experimental/Animation/SkeletalMeshModelingTools/Source/SkeletalMeshModelingTools/Private/SkeletalMeshModelingToolsModule.cpp b/Engine/Plugins/Experimental/Animation/SkeletalMeshModelingTools/Source/SkeletalMeshModelingTools/Private/SkeletalMeshModelingToolsModule.cpp
--- a/Engine/Plugins/Experimental/Animation/SkeletalMeshModelingTools/Source/SkeletalMeshModelingTools/Private/SkeletalMeshModelingToolsModule.cpp
+++ b/Engine/Plugins/Experimental/Animation/SkeletalMeshModelingTools/Source/SkeletalMeshModelingTools/Private/SkeletalMeshModelingToolsModule.cpp
@@ -12,6 +12,7 @@
 #include "Engine/SkeletalMesh.h"
 #include "ISkeletalMeshEditor.h"
 #include "ISkeletalMeshEditorModule.h"
+#include "ModelingToolsEditorModeStyle.h"
 #include "ModelingToolsManagerActions.h"
 #include "Modules/ModuleManager.h"
 #include "SkeletalMeshToolMenuContext.h"
@@ -164,6 +165,7 @@ void FSkeletalMeshModelingToolsModule::OnToggleEditingToolsMode(TWeakPtr<ISkelet
 void FSkeletalMeshModelingToolsModule::OnPostEngineInit()
 {
 	FSkeletalMeshModelingToolsActionCommands::RegisterAllToolActions();
+	FModelingToolsEditorModeStyle::Initialize();
 	FModelingToolsManagerCommands::Register();
 }
 
