Renaming ISkeletalMeshEditionInterface to *EditingInterface

diff --git a/Engine/Plugins/Experimental/Animation/SkeletalMeshModelingTools/Source/SkeletalMeshModelingTools/Private/SkeletalMeshModelingToolsEditorMode.cpp b/Engine/Plugins/Experimental/Animation/SkeletalMeshModelingTools/Source/SkeletalMeshModelingTools/Private/SkeletalMeshModelingToolsEditorMode.cpp
--- a/Engine/Plugins/Experimental/Animation/SkeletalMeshModelingTools/Source/SkeletalMeshModelingTools/Private/SkeletalMeshModelingToolsEditorMode.cpp
+++ b/Engine/Plugins/Experimental/Animation/SkeletalMeshModelingTools/Source/SkeletalMeshModelingTools/Private/SkeletalMeshModelingToolsEditorMode.cpp
@@ -368,18 +368,18 @@ void USkeletalMeshModelingToolsEditorMode::SetEditorBinding(TSharedPtr<ISkeletal
 	Binding = InBinding;
 }
 
-ISkeletalMeshEditionInterface* USkeletalMeshModelingToolsEditorMode::GetSkeletonInterface(UInteractiveTool* InTool)
+ISkeletalMeshEditingInterface* USkeletalMeshModelingToolsEditorMode::GetSkeletonInterface(UInteractiveTool* InTool)
 {
-	if (!IsValid(InTool) || !InTool->Implements<USkeletalMeshEditionInterface>())
+	if (!IsValid(InTool) || !InTool->Implements<USkeletalMeshEditingInterface>())
 	{
 		return nullptr;
 	}
-	return static_cast<ISkeletalMeshEditionInterface*>(InTool->GetInterfaceAddress(USkeletalMeshEditionInterface::StaticClass()));
+	return static_cast<ISkeletalMeshEditingInterface*>(InTool->GetInterfaceAddress(USkeletalMeshEditingInterface::StaticClass()));
 }
 
 void USkeletalMeshModelingToolsEditorMode::ConnectTool(UInteractiveTool* InTool)
 {
-	ISkeletalMeshEditionInterface* SkeletonInterface = GetSkeletonInterface(InTool);
+	ISkeletalMeshEditingInterface* SkeletonInterface = GetSkeletonInterface(InTool);
 	if (!SkeletonInterface)
 	{
 		return;
@@ -428,7 +428,7 @@ void USkeletalMeshModelingToolsEditorMode::DisconnectTool(UInteractiveTool* InTo
 
 	if (FromToolNotifierHandle.IsValid())
 	{
-		ISkeletalMeshEditionInterface* SkeletonInterface = GetSkeletonInterface(InTool);
+		ISkeletalMeshEditingInterface* SkeletonInterface = GetSkeletonInterface(InTool);
 		check(SkeletonInterface);
 		SkeletonInterface->GetNotifier().Delegate().Remove(FromToolNotifierHandle);
 		FromToolNotifierHandle.Reset();
@@ -439,7 +439,7 @@ void USkeletalMeshModelingToolsEditorMode::DisconnectTool(UInteractiveTool* InTo
 bool USkeletalMeshModelingToolsEditorMode::NeedsTransformGizmo() const
 {
 	UInteractiveTool* Tool = GetToolManager()->GetActiveTool(EToolSide::Mouse);
-	if (const ISkeletalMeshEditionInterface* SkeletonInterface = GetSkeletonInterface(Tool))
+	if (const ISkeletalMeshEditingInterface* SkeletonInterface = GetSkeletonInterface(Tool))
 	{
 		return !SkeletonInterface->GetSelectedBones().IsEmpty();
 	}
diff --git a/Engine/Plugins/Experimental/Animation/SkeletalMeshModelingTools/Source/SkeletalMeshModelingTools/Private/SkeletalMeshModelingToolsEditorMode.h b/Engine/Plugins/Experimental/Animation/SkeletalMeshModelingTools/Source/SkeletalMeshModelingTools/Private/SkeletalMeshModelingToolsEditorMode.h
--- a/Engine/Plugins/Experimental/Animation/SkeletalMeshModelingTools/Source/SkeletalMeshModelingTools/Private/SkeletalMeshModelingToolsEditorMode.h
+++ b/Engine/Plugins/Experimental/Animation/SkeletalMeshModelingTools/Source/SkeletalMeshModelingTools/Private/SkeletalMeshModelingToolsEditorMode.h
@@ -12,7 +12,7 @@ class FSkeletalMeshModelingToolsEditorModeToolkit;
 class UEdModeInteractiveToolsContext;
 class ISkeletalMeshNotifier;
 class ISkeletalMeshEditorBinding;
-class ISkeletalMeshEditionInterface;
+class ISkeletalMeshEditingInterface;
 class HHitProxy;
 class UDebugSkelMeshComponent;
 
@@ -55,7 +55,7 @@ private:
 	// Stylus support is currently disabled; this is left in for reference if/when it is brought back
 	//TUniquePtr<FStylusStateTracker> StylusStateTracker;
 
-	static ISkeletalMeshEditionInterface* GetSkeletonInterface(UInteractiveTool* InTool);
+	static ISkeletalMeshEditingInterface* GetSkeletonInterface(UInteractiveTool* InTool);
 
 	UDebugSkelMeshComponent* GetSkelMeshComponent() const;
 
diff --git a/Engine/Plugins/Experimental/MeshModelingToolsetExp/Source/MeshModelingToolsEditorOnlyExp/Private/SkeletalMesh/SkeletalMeshEditionInterface.cpp b/Engine/Plugins/Experimental/MeshModelingToolsetExp/Source/MeshModelingToolsEditorOnlyExp/Private/SkeletalMesh/SkeletalMeshEditionInterface.cpp
--- a/Engine/Plugins/Experimental/MeshModelingToolsetExp/Source/MeshModelingToolsEditorOnlyExp/Private/SkeletalMesh/SkeletalMeshEditionInterface.cpp
+++ b/Engine/Plugins/Experimental/MeshModelingToolsetExp/Source/MeshModelingToolsEditorOnlyExp/Private/SkeletalMesh/SkeletalMeshEditionInterface.cpp
@@ -4,32 +4,32 @@
 
 IMPLEMENT_HIT_PROXY(HBoneHitProxy, HHitProxy)
 
-ISkeletalMeshNotifier& ISkeletalMeshEditionInterface::GetNotifier()
+ISkeletalMeshNotifier& ISkeletalMeshEditingInterface::GetNotifier()
 {
 	if (!Notifier)
 	{
-		Notifier.Reset(new FSkeletalMeshToolNotifier(TWeakInterfacePtr<ISkeletalMeshEditionInterface>(this)));
+		Notifier.Reset(new FSkeletalMeshToolNotifier(TWeakInterfacePtr<ISkeletalMeshEditingInterface>(this)));
 	}
 	
 	return *Notifier;
 }
 
-bool ISkeletalMeshEditionInterface::NeedsNotification() const
+bool ISkeletalMeshEditingInterface::NeedsNotification() const
 {
 	return Notifier && Notifier->Delegate().IsBound(); 
 }
 
-void ISkeletalMeshEditionInterface::BindTo(TSharedPtr<ISkeletalMeshEditorBinding> InBinding)
+void ISkeletalMeshEditingInterface::BindTo(TSharedPtr<ISkeletalMeshEditorBinding> InBinding)
 {
 	Binding = InBinding;
 }
 
-void ISkeletalMeshEditionInterface::Unbind()
+void ISkeletalMeshEditingInterface::Unbind()
 {
 	Binding.Reset();
 }
 
-TOptional<FName> ISkeletalMeshEditionInterface::GetBoneName(HHitProxy* InHitProxy) const
+TOptional<FName> ISkeletalMeshEditingInterface::GetBoneName(HHitProxy* InHitProxy) const
 {
 	if (const HBoneHitProxy* BoneProxy = HitProxyCast<HBoneHitProxy>(InHitProxy))
 	{
@@ -39,13 +39,13 @@ TOptional<FName> ISkeletalMeshEditionInterface::GetBoneName(HHitProxy* InHitProx
 	return Binding.IsValid() && Binding.Pin()->GetNameFunction() ? Binding.Pin()->GetNameFunction()(InHitProxy) : TOptional<FName>();
 }
 
-TArray<FName> ISkeletalMeshEditionInterface::GetSelectedBones() const
+TArray<FName> ISkeletalMeshEditingInterface::GetSelectedBones() const
 {
 	static const TArray<FName> Dummy;
 	return Binding.IsValid() ? Binding.Pin()->GetSelectedBones() : Dummy;
 }
 
-FSkeletalMeshToolNotifier::FSkeletalMeshToolNotifier(TWeakInterfacePtr<ISkeletalMeshEditionInterface> InInterface)
+FSkeletalMeshToolNotifier::FSkeletalMeshToolNotifier(TWeakInterfacePtr<ISkeletalMeshEditingInterface> InInterface)
 	: ISkeletalMeshNotifier()
 	, Interface(InInterface)
 {}
diff --git a/Engine/Plugins/Experimental/MeshModelingToolsetExp/Source/MeshModelingToolsEditorOnlyExp/Public/SkeletalMesh/SkeletalMeshEditionInterface.h b/Engine/Plugins/Experimental/MeshModelingToolsetExp/Source/MeshModelingToolsEditorOnlyExp/Public/SkeletalMesh/SkeletalMeshEditionInterface.h
--- a/Engine/Plugins/Experimental/MeshModelingToolsetExp/Source/MeshModelingToolsEditorOnlyExp/Public/SkeletalMesh/SkeletalMeshEditionInterface.h
+++ b/Engine/Plugins/Experimental/MeshModelingToolsetExp/Source/MeshModelingToolsEditorOnlyExp/Public/SkeletalMesh/SkeletalMeshEditionInterface.h
@@ -23,7 +23,7 @@ struct FInputDeviceRay;
  */
 
 UINTERFACE()
-class MESHMODELINGTOOLSEDITORONLYEXP_API USkeletalMeshEditionInterface : public UInterface
+class MESHMODELINGTOOLSEDITORONLYEXP_API USkeletalMeshEditingInterface : public UInterface
 {
 	GENERATED_BODY()
 };
@@ -32,7 +32,7 @@ class MESHMODELINGTOOLSEDITORONLYEXP_API USkeletalMeshEditionInterface : public
  * ISkeletalMeshEditionInterface
  */
 
-class MESHMODELINGTOOLSEDITORONLYEXP_API ISkeletalMeshEditionInterface
+class MESHMODELINGTOOLSEDITORONLYEXP_API ISkeletalMeshEditingInterface
 {
 	GENERATED_BODY()
 
@@ -64,11 +64,11 @@ private:
 class FSkeletalMeshToolNotifier: public ISkeletalMeshNotifier
 {
 public:
-	FSkeletalMeshToolNotifier(TWeakInterfacePtr<ISkeletalMeshEditionInterface> InInterface);
+	FSkeletalMeshToolNotifier(TWeakInterfacePtr<ISkeletalMeshEditingInterface> InInterface);
 	virtual void HandleNotification(const TArray<FName>& BoneNames, const ESkeletalMeshNotifyType InNotifyType) override;
 	
 protected:
-	TWeakInterfacePtr<ISkeletalMeshEditionInterface> Interface;
+	TWeakInterfacePtr<ISkeletalMeshEditingInterface> Interface;
 };
 
 /**
diff --git a/Engine/Plugins/Experimental/MeshModelingToolsetExp/Source/MeshModelingToolsEditorOnlyExp/Public/SkeletalMesh/SkeletonEditingTool.h b/Engine/Plugins/Experimental/MeshModelingToolsetExp/Source/MeshModelingToolsEditorOnlyExp/Public/SkeletalMesh/SkeletonEditingTool.h
--- a/Engine/Plugins/Experimental/MeshModelingToolsetExp/Source/MeshModelingToolsEditorOnlyExp/Public/SkeletalMesh/SkeletonEditingTool.h
+++ b/Engine/Plugins/Experimental/MeshModelingToolsetExp/Source/MeshModelingToolsEditorOnlyExp/Public/SkeletalMesh/SkeletonEditingTool.h
@@ -96,7 +96,7 @@ UCLASS()
 class MESHMODELINGTOOLSEDITORONLYEXP_API USkeletonEditingTool :
 	public USingleSelectionTool,
 	public IClickDragBehaviorTarget,
-	public ISkeletalMeshEditionInterface
+	public ISkeletalMeshEditingInterface
 {
 	GENERATED_BODY()
 
diff --git a/Engine/Plugins/Experimental/MeshModelingToolsetExp/Source/MeshModelingToolsEditorOnlyExp/Public/SkeletalMesh/SkinWeightsPaintTool.h b/Engine/Plugins/Experimental/MeshModelingToolsetExp/Source/MeshModelingToolsEditorOnlyExp/Public/SkeletalMesh/SkinWeightsPaintTool.h
--- a/Engine/Plugins/Experimental/MeshModelingToolsetExp/Source/MeshModelingToolsEditorOnlyExp/Public/SkeletalMesh/SkinWeightsPaintTool.h
+++ b/Engine/Plugins/Experimental/MeshModelingToolsetExp/Source/MeshModelingToolsEditorOnlyExp/Public/SkeletalMesh/SkinWeightsPaintTool.h
@@ -287,7 +287,7 @@ public:
 
 // An interactive tool for painting and editing skin weights.
 UCLASS()
-class MESHMODELINGTOOLSEDITORONLYEXP_API USkinWeightsPaintTool : public UDynamicMeshBrushTool, public ISkeletalMeshEditionInterface
+class MESHMODELINGTOOLSEDITORONLYEXP_API USkinWeightsPaintTool : public UDynamicMeshBrushTool, public ISkeletalMeshEditingInterface
 {
 	GENERATED_BODY()
 
