Skeleton Editing Tools: the interface between the skeletal mesh editor and the tools now use a context object to get/set info from the editor
    - hide bones when switching to skeleton editing mode

diff --git a/Engine/Plugins/Experimental/Animation/SkeletalMeshModelingTools/Source/SkeletalMeshModelingTools/Private/SkeletalMeshEditorUtils.cpp b/Engine/Plugins/Experimental/Animation/SkeletalMeshModelingTools/Source/SkeletalMeshModelingTools/Private/SkeletalMeshEditorUtils.cpp
new file mode 100644
--- /dev/null
+++ b/Engine/Plugins/Experimental/Animation/SkeletalMeshModelingTools/Source/SkeletalMeshModelingTools/Private/SkeletalMeshEditorUtils.cpp
@@ -0,0 +1,188 @@
+﻿// Copyright Epic Games, Inc. All Rights Reserved.
+
+#include "SkeletalMeshEditorUtils.h"
+
+#include "ContextObjectStore.h"
+#include "InteractiveToolsContext.h"
+#include "IPersonaToolkit.h"
+#include "ISkeletalMeshEditor.h"
+
+#include UE_INLINE_GENERATED_CPP_BY_NAME(SkeletalMeshEditorUtils)
+
+bool UE::SkeletalMeshEditorUtils::RegisterEditorContextObject(UInteractiveToolsContext* ToolsContext)
+{
+	if (ensure(ToolsContext))
+	{
+		const USkeletalMeshEditorContextObject* Found = ToolsContext->ContextObjectStore->FindContext<USkeletalMeshEditorContextObject>();
+		if (Found)
+		{
+			return true;
+		}
+		
+		USkeletalMeshEditorContextObject* ContextObject = NewObject<USkeletalMeshEditorContextObject>(ToolsContext->ToolManager);
+		if (ensure(ContextObject))
+		{
+			ContextObject->Register(ToolsContext->ToolManager);
+			return true;
+		}
+	}
+	return false;
+}
+
+bool UE::SkeletalMeshEditorUtils::UnregisterEditorContextObject(UInteractiveToolsContext* ToolsContext)
+{
+	if (ensure(ToolsContext))
+	{
+		USkeletalMeshEditorContextObject* Found = ToolsContext->ContextObjectStore->FindContext<USkeletalMeshEditorContextObject>();
+		if (Found != nullptr)
+		{
+			Found->Unregister(ToolsContext->ToolManager);
+			ToolsContext->ContextObjectStore->RemoveContextObject(Found);
+		}
+		return true;
+	}
+	return false;
+}
+
+USkeletalMeshEditorContextObject* UE::SkeletalMeshEditorUtils::GetEditorContextObject(UInteractiveToolsContext* ToolsContext)
+{
+	return ToolsContext->ContextObjectStore->FindContext<USkeletalMeshEditorContextObject>();
+}
+
+void USkeletalMeshEditorContextObject::Register(UInteractiveToolManager* InToolManager)
+{
+	if (ensure(!bRegistered) == false)
+	{
+		return;
+	}
+
+	InToolManager->GetContextObjectStore()->AddContextObject(this);
+	bRegistered = true;
+}
+
+void USkeletalMeshEditorContextObject::Unregister(UInteractiveToolManager* InToolManager)
+{
+	ensure(bRegistered);
+	
+	InToolManager->GetContextObjectStore()->RemoveContextObject(this);
+
+	Bindings.Reset();
+	
+	bRegistered = false;
+}
+
+void USkeletalMeshEditorContextObject::Init(const TWeakPtr<ISkeletalMeshEditor>& InEditor)
+{
+	Editor = InEditor;
+	Bindings.Reset();
+}
+
+void USkeletalMeshEditorContextObject::HideSkeleton()
+{
+	if (!Editor.IsValid())
+	{
+		return;
+	}
+
+	const TSharedPtr<ISkeletalMeshEditor> SkeletalMeshEditor = Editor.Pin();
+	UDebugSkelMeshComponent* SkeletalMeshComponent = SkeletalMeshEditor->GetPersonaToolkit()->GetPreviewMeshComponent();
+	if (!SkeletalMeshComponent)
+	{
+		return;
+	}
+	
+	SkeletonDrawMode = SkeletalMeshComponent->SkeletonDrawMode;
+	SkeletalMeshComponent->SkeletonDrawMode = ESkeletonDrawMode::Hidden;
+}
+
+void USkeletalMeshEditorContextObject::ShowSkeleton()
+{
+	if (!Editor.IsValid())
+	{
+		return;
+	}
+	
+	const TSharedPtr<ISkeletalMeshEditor> SkeletalMeshEditor = Editor.Pin();
+	UDebugSkelMeshComponent* SkeletalMeshComponent = SkeletalMeshEditor->GetPersonaToolkit()->GetPreviewMeshComponent();
+	if (!SkeletalMeshComponent)
+	{
+		return;
+	}
+	
+	SkeletalMeshComponent->SkeletonDrawMode = SkeletonDrawMode;
+	SkeletonDrawMode = ESkeletonDrawMode::Default;
+}
+
+void USkeletalMeshEditorContextObject::BindTo(ISkeletalMeshEditingInterface* InEditingInterface)
+{
+	if (!InEditingInterface || !Editor.IsValid())
+	{
+		return;
+	}
+	
+	if (Bindings.Contains(InEditingInterface))
+	{
+		return;
+	}
+	
+	TSharedPtr<ISkeletalMeshEditorBinding> Binding = Editor.Pin()->GetBinding();
+	if (Binding.IsValid())
+	{
+		InEditingInterface->BindTo(Binding);
+	
+		FBindData BindData;
+	
+		// connect external interface to tool (ie skeletal mesh editor -> tool)
+		BindData.ToToolNotifierHandle = Binding->GetNotifier().Delegate().AddLambda(
+			[InEditingInterface](const TArray<FName>& BoneNames, const ESkeletalMeshNotifyType InNotifyType)
+		{
+			InEditingInterface->GetNotifier().HandleNotification(BoneNames, InNotifyType);
+		});
+	
+		// connect too to external interface (ie tool -> skeletal mesh editor)
+		BindData.FromToolNotifierHandle = InEditingInterface->GetNotifier().Delegate().AddLambda(
+			[Binding](const TArray<FName>& BoneNames, const ESkeletalMeshNotifyType InNotifyType)
+		{
+			Binding->GetNotifier().HandleNotification(BoneNames, InNotifyType);
+		});
+		Bindings.Emplace(InEditingInterface, BindData);
+	
+		InEditingInterface->GetNotifier().HandleNotification(Binding->GetSelectedBones(), ESkeletalMeshNotifyType::BonesSelected);
+	}
+}
+
+void USkeletalMeshEditorContextObject::UnbindFrom(ISkeletalMeshEditingInterface* InEditingInterface)
+{
+	if (!InEditingInterface)
+	{
+		return;
+	}
+	
+	if (FBindData* BindData = Bindings.Find(InEditingInterface))
+	{
+		if (BindData->ToToolNotifierHandle.IsValid())
+		{
+			TSharedPtr<ISkeletalMeshEditorBinding> Binding = GetBinding();
+			if (Binding.IsValid())
+			{
+				Binding->GetNotifier().Delegate().Remove(BindData->ToToolNotifierHandle);
+			}
+			BindData->ToToolNotifierHandle.Reset();
+		}
+	
+		if (BindData->FromToolNotifierHandle.IsValid())
+		{
+			InEditingInterface->GetNotifier().Delegate().Remove(BindData->FromToolNotifierHandle);
+			InEditingInterface->Unbind();
+			
+			BindData->FromToolNotifierHandle.Reset();
+		}
+	
+		Bindings.Remove(InEditingInterface);
+	}
+}
+
+TSharedPtr<ISkeletalMeshEditorBinding> USkeletalMeshEditorContextObject::GetBinding() const
+{
+	return Editor.IsValid() ? Editor.Pin()->GetBinding() : nullptr;
+}
\ No newline at end of file
diff --git a/Engine/Plugins/Experimental/Animation/SkeletalMeshModelingTools/Source/SkeletalMeshModelingTools/Private/SkeletalMeshEditorUtils.h b/Engine/Plugins/Experimental/Animation/SkeletalMeshModelingTools/Source/SkeletalMeshModelingTools/Private/SkeletalMeshEditorUtils.h
new file mode 100644
--- /dev/null
+++ b/Engine/Plugins/Experimental/Animation/SkeletalMeshModelingTools/Source/SkeletalMeshModelingTools/Private/SkeletalMeshEditorUtils.h
@@ -0,0 +1,68 @@
+﻿// Copyright Epic Games, Inc. All Rights Reserved.
+
+#pragma once
+
+#include "Animation/DebugSkelMeshComponent.h"
+#include "SkeletalMesh/SkeletalMeshEditionInterface.h"
+
+#include "UObject/Object.h"
+#include "UObject/WeakInterfacePtr.h"
+
+#include "SkeletalMeshEditorUtils.generated.h"
+
+class UInteractiveToolsContext;
+class UInteractiveToolManager;
+class ISkeletalMeshEditor;
+class USkeletalMeshEditorContextObject;
+
+namespace UE
+{
+	
+namespace SkeletalMeshEditorUtils
+{
+	SKELETALMESHMODELINGTOOLS_API bool RegisterEditorContextObject(UInteractiveToolsContext* InToolsContext);
+	SKELETALMESHMODELINGTOOLS_API bool UnregisterEditorContextObject(UInteractiveToolsContext* InToolsContext);
+
+	SKELETALMESHMODELINGTOOLS_API USkeletalMeshEditorContextObject* GetEditorContextObject(UInteractiveToolsContext* InToolsContext);
+}
+	
+}
+
+/**
+ * USkeletalMeshEditorContextObject
+ */
+
+UCLASS()
+class SKELETALMESHMODELINGTOOLS_API USkeletalMeshEditorContextObject : public USkeletalMeshEditorContextObjectBase
+{
+	GENERATED_BODY()
+
+public:
+	void Register(UInteractiveToolManager* InToolManager);
+	void Unregister(UInteractiveToolManager* InToolManager);
+
+	void Init(const TWeakPtr<ISkeletalMeshEditor>& InEditor);
+
+	virtual void HideSkeleton() override;
+	virtual void ShowSkeleton() override;
+
+	virtual void BindTo(ISkeletalMeshEditingInterface* InEditingInterface) override;
+	virtual void UnbindFrom(ISkeletalMeshEditingInterface* InEditingInterface) override;
+	
+private:
+	TWeakPtr<ISkeletalMeshEditor> Editor = nullptr;
+
+	UPROPERTY()
+	ESkeletonDrawMode SkeletonDrawMode = ESkeletonDrawMode::Default;
+
+	TSharedPtr<ISkeletalMeshEditorBinding> GetBinding() const;
+
+	struct FBindData
+	{
+		FDelegateHandle ToToolNotifierHandle;
+		FDelegateHandle FromToolNotifierHandle;		
+	};
+	TMap< ISkeletalMeshEditingInterface*, FBindData > Bindings;
+	
+	bool bRegistered = false;
+};
diff --git a/Engine/Plugins/Experimental/Animation/SkeletalMeshModelingTools/Source/SkeletalMeshModelingTools/Private/SkeletalMeshModelingToolsEditorMode.cpp b/Engine/Plugins/Experimental/Animation/SkeletalMeshModelingTools/Source/SkeletalMeshModelingTools/Private/SkeletalMeshModelingToolsEditorMode.cpp
--- a/Engine/Plugins/Experimental/Animation/SkeletalMeshModelingTools/Source/SkeletalMeshModelingTools/Private/SkeletalMeshModelingToolsEditorMode.cpp
+++ b/Engine/Plugins/Experimental/Animation/SkeletalMeshModelingTools/Source/SkeletalMeshModelingTools/Private/SkeletalMeshModelingToolsEditorMode.cpp
@@ -31,6 +31,7 @@
 #include "EditMeshPolygonsTool.h"
 #include "EditorInteractiveGizmoManager.h"
 #include "HoleFillTool.h"
+#include "ISkeletalMeshEditor.h"
 #include "LatticeDeformerTool.h"
 #include "MeshAttributePaintTool.h"
 #include "MeshSpaceDeformerTool.h"
@@ -42,6 +43,7 @@
 #include "RemeshMeshTool.h"
 #include "RemoveOccludedTrianglesTool.h"
 #include "SimplifyMeshTool.h"
+#include "SkeletalMeshEditorUtils.h"
 #include "SmoothMeshTool.h"
 #include "ToolTargetManager.h"
 #include "WeldMeshEdgesTool.h"
@@ -190,6 +192,7 @@ void USkeletalMeshModelingToolsEditorMode::Enter()
 	// register gizmo helper
 	UE::TransformGizmoUtil::RegisterTransformGizmoContextObject(InteractiveToolsContext);
 	UE::SkeletalMeshGizmoUtils::RegisterTransformGizmoContextObject(InteractiveToolsContext);
+	UE::SkeletalMeshEditorUtils::RegisterEditorContextObject(InteractiveToolsContext);
 
 	const FModelingToolsManagerCommands& ToolManagerCommands = FModelingToolsManagerCommands::Get();
 
@@ -251,6 +254,8 @@ void USkeletalMeshModelingToolsEditorMode::Exit()
 	UEditorInteractiveToolsContext* InteractiveToolsContext = GetInteractiveToolsContext();
 	UE::TransformGizmoUtil::DeregisterTransformGizmoContextObject(InteractiveToolsContext);
 	UE::SkeletalMeshGizmoUtils::UnregisterTransformGizmoContextObject(InteractiveToolsContext);
+	UE::SkeletalMeshEditorUtils::UnregisterEditorContextObject(InteractiveToolsContext);
+	
 #if ENABLE_STYLUS_SUPPORT
 	StylusStateTracker = nullptr;
 #endif
@@ -352,88 +357,35 @@ bool USkeletalMeshModelingToolsEditorMode::ComputeBoundingBoxForViewportFocus(AA
 void USkeletalMeshModelingToolsEditorMode::OnToolStarted(UInteractiveToolManager* Manager, UInteractiveTool* Tool)
 {
 	FSkeletalMeshModelingToolsActionCommands::UpdateToolCommandBinding(Tool, Toolkit->GetToolkitCommands(), false);
-	
-	ConnectTool(Tool);
 }
 
 void USkeletalMeshModelingToolsEditorMode::OnToolEnded(UInteractiveToolManager* Manager, UInteractiveTool* Tool)
 {
 	FSkeletalMeshModelingToolsActionCommands::UpdateToolCommandBinding(Tool, Toolkit->GetToolkitCommands(), true);
-	
-	DisconnectTool(Tool);
-}
-
-void USkeletalMeshModelingToolsEditorMode::SetEditorBinding(TSharedPtr<ISkeletalMeshEditorBinding> InBinding)
-{
-	Binding = InBinding;
 }
 
-ISkeletalMeshEditingInterface* USkeletalMeshModelingToolsEditorMode::GetSkeletonInterface(UInteractiveTool* InTool)
+void USkeletalMeshModelingToolsEditorMode::SetEditorBinding(const TWeakPtr<ISkeletalMeshEditor>& InSkeletalMeshEditor)
 {
-	if (!IsValid(InTool) || !InTool->Implements<USkeletalMeshEditingInterface>())
-	{
-		return nullptr;
-	}
-	return static_cast<ISkeletalMeshEditingInterface*>(InTool->GetInterfaceAddress(USkeletalMeshEditingInterface::StaticClass()));
-}
-
-void USkeletalMeshModelingToolsEditorMode::ConnectTool(UInteractiveTool* InTool)
-{
-	ISkeletalMeshEditingInterface* SkeletonInterface = GetSkeletonInterface(InTool);
-	if (!SkeletonInterface)
+	if (!InSkeletalMeshEditor.IsValid())
 	{
 		return;
 	}
+	
+	Binding = InSkeletalMeshEditor.Pin()->GetBinding();
 
-	// binding
-	if (Binding.IsValid())
+	if (USkeletalMeshEditorContextObject* ContextObject = UE::SkeletalMeshEditorUtils::GetEditorContextObject(GetInteractiveToolsContext()))
 	{
-		TSharedPtr<ISkeletalMeshEditorBinding> BindingPtr = Binding.Pin();
-
-		SkeletonInterface->BindTo(BindingPtr);
-
-		// connect external interface to tool (ie skeletal mesh editor -> tool)
-		if (!ToToolNotifierHandle.IsValid())
-		{
-			ToToolNotifierHandle = BindingPtr->GetNotifier().Delegate().AddLambda([SkeletonInterface](const TArray<FName>& BoneNames, const ESkeletalMeshNotifyType InNotifyType)
-			{
-				SkeletonInterface->GetNotifier().HandleNotification(BoneNames, InNotifyType);
-			});
-		}
-
-		// connect too to external interface (ie tool -> skeletal mesh editor)
-		if (!FromToolNotifierHandle.IsValid())
-		{
-			FromToolNotifierHandle = SkeletonInterface->GetNotifier().Delegate().AddLambda([BindingPtr](const TArray<FName>& BoneNames, const ESkeletalMeshNotifyType InNotifyType)
-			{
-				BindingPtr->GetNotifier().HandleNotification(BoneNames, InNotifyType);
-			});
-		}
-
-		SkeletonInterface->GetNotifier().HandleNotification(BindingPtr->GetSelectedBones(), ESkeletalMeshNotifyType::BonesSelected);
+		ContextObject->Init(InSkeletalMeshEditor);
 	}
 }
 
-void USkeletalMeshModelingToolsEditorMode::DisconnectTool(UInteractiveTool* InTool)
+ISkeletalMeshEditingInterface* USkeletalMeshModelingToolsEditorMode::GetSkeletonInterface(UInteractiveTool* InTool)
 {
-	if (ToToolNotifierHandle.IsValid())
-	{
-		// check that
-		if (Binding.IsValid())
-		{
-			Binding.Pin()->GetNotifier().Delegate().Remove(ToToolNotifierHandle);
-		}
-		ToToolNotifierHandle.Reset();
-	}
-
-	if (FromToolNotifierHandle.IsValid())
+	if (!IsValid(InTool) || !InTool->Implements<USkeletalMeshEditingInterface>())
 	{
-		ISkeletalMeshEditingInterface* SkeletonInterface = GetSkeletonInterface(InTool);
-		check(SkeletonInterface);
-		SkeletonInterface->GetNotifier().Delegate().Remove(FromToolNotifierHandle);
-		FromToolNotifierHandle.Reset();
-		SkeletonInterface->Unbind();
+		return nullptr;
 	}
+	return static_cast<ISkeletalMeshEditingInterface*>(InTool->GetInterfaceAddress(USkeletalMeshEditingInterface::StaticClass()));
 }
 
 bool USkeletalMeshModelingToolsEditorMode::NeedsTransformGizmo() const
diff --git a/Engine/Plugins/Experimental/Animation/SkeletalMeshModelingTools/Source/SkeletalMeshModelingTools/Private/SkeletalMeshModelingToolsEditorMode.h b/Engine/Plugins/Experimental/Animation/SkeletalMeshModelingTools/Source/SkeletalMeshModelingTools/Private/SkeletalMeshModelingToolsEditorMode.h
--- a/Engine/Plugins/Experimental/Animation/SkeletalMeshModelingTools/Source/SkeletalMeshModelingTools/Private/SkeletalMeshModelingToolsEditorMode.h
+++ b/Engine/Plugins/Experimental/Animation/SkeletalMeshModelingTools/Source/SkeletalMeshModelingTools/Private/SkeletalMeshModelingToolsEditorMode.h
@@ -15,6 +15,7 @@ class ISkeletalMeshEditorBinding;
 class ISkeletalMeshEditingInterface;
 class HHitProxy;
 class UDebugSkelMeshComponent;
+class ISkeletalMeshEditor;
 
 UCLASS()
 class USkeletalMeshModelingToolsEditorMode : 
@@ -45,7 +46,7 @@ public:
 	virtual bool UsesToolkits() const override { return true; }
 
 	// binding
-	void SetEditorBinding(TSharedPtr<ISkeletalMeshEditorBinding> InBinding);
+	void SetEditorBinding(const TWeakPtr<ISkeletalMeshEditor>& InSkeletalMeshEditor);
 		
 protected:
 	virtual void OnToolStarted(UInteractiveToolManager* Manager, UInteractiveTool* Tool) override;
@@ -59,9 +60,6 @@ private:
 
 	UDebugSkelMeshComponent* GetSkelMeshComponent() const;
 
-	void ConnectTool(UInteractiveTool* InTool);
-	void DisconnectTool(UInteractiveTool* InTool);
-
 	bool NeedsTransformGizmo() const;
 
 	FDelegateHandle ToToolNotifierHandle;
diff --git a/Engine/Plugins/Experimental/Animation/SkeletalMeshModelingTools/Source/SkeletalMeshModelingTools/Private/SkeletalMeshModelingToolsModule.cpp b/Engine/Plugins/Experimental/Animation/SkeletalMeshModelingTools/Source/SkeletalMeshModelingTools/Private/SkeletalMeshModelingToolsModule.cpp
--- a/Engine/Plugins/Experimental/Animation/SkeletalMeshModelingTools/Source/SkeletalMeshModelingTools/Private/SkeletalMeshModelingToolsModule.cpp
+++ b/Engine/Plugins/Experimental/Animation/SkeletalMeshModelingTools/Source/SkeletalMeshModelingTools/Private/SkeletalMeshModelingToolsModule.cpp
@@ -152,7 +152,7 @@ void FSkeletalMeshModelingToolsModule::OnToggleEditingToolsMode(TWeakPtr<ISkelet
 			UEdMode* ActiveMode = EditorModeManager.GetActiveScriptableMode(USkeletalMeshModelingToolsEditorMode::Id);
 			if (USkeletalMeshModelingToolsEditorMode* SkeletalMeshMode = CastChecked<USkeletalMeshModelingToolsEditorMode>(ActiveMode))
 			{
-				SkeletalMeshMode->SetEditorBinding(SkeletalMeshEditor->GetBinding());
+				SkeletalMeshMode->SetEditorBinding(SkeletalMeshEditor);
 			}
 		}
 		else
diff --git a/Engine/Plugins/Experimental/MeshModelingToolsetExp/Source/MeshModelingToolsEditorOnlyExp/Private/SkeletalMesh/SkeletonEditingTool.cpp b/Engine/Plugins/Experimental/MeshModelingToolsetExp/Source/MeshModelingToolsEditorOnlyExp/Private/SkeletalMesh/SkeletonEditingTool.cpp
--- a/Engine/Plugins/Experimental/MeshModelingToolsetExp/Source/MeshModelingToolsEditorOnlyExp/Private/SkeletalMesh/SkeletonEditingTool.cpp
+++ b/Engine/Plugins/Experimental/MeshModelingToolsetExp/Source/MeshModelingToolsEditorOnlyExp/Private/SkeletalMesh/SkeletonEditingTool.cpp
@@ -93,6 +93,7 @@ void USkeletonEditingTool::Init(const FToolBuilderState& InSceneState)
 	const UContextObjectStore* ContextObjectStore = InSceneState.ToolManager->GetContextObjectStore();
 	ViewContext = ContextObjectStore->FindContext<UGizmoViewContext>();
 	GizmoContext = ContextObjectStore->FindContext<USkeletalMeshGizmoContextObjectBase>();
+	EditorContext = ContextObjectStore->FindContext<USkeletalMeshEditorContextObjectBase>();
 }
 
 void USkeletonEditingTool::Setup()
@@ -172,6 +173,12 @@ void USkeletonEditingTool::Setup()
 		AddInputBehavior(ClickDragBehavior);
 	}
 
+	if (EditorContext.IsValid())
+	{
+		EditorContext->HideSkeleton();
+		EditorContext->BindTo(this);
+	}
+	
 	// setup gizmo
 	if (GizmoContext.IsValid())
 	{
@@ -332,6 +339,12 @@ void USkeletonEditingTool::Shutdown(EToolShutdownType ShutdownType)
 	{
 		GizmoWrapper->Clear();
 	}
+	
+	if (EditorContext.IsValid())
+	{
+		EditorContext->ShowSkeleton();
+		EditorContext->UnbindFrom(this);
+	}
 }
 
 void USkeletonEditingTool::RegisterActions(FInteractiveToolActionSet& ActionSet)
diff --git a/Engine/Plugins/Experimental/MeshModelingToolsetExp/Source/MeshModelingToolsEditorOnlyExp/Private/SkeletalMesh/SkinWeightsBindingTool.cpp b/Engine/Plugins/Experimental/MeshModelingToolsetExp/Source/MeshModelingToolsEditorOnlyExp/Private/SkeletalMesh/SkinWeightsBindingTool.cpp
--- a/Engine/Plugins/Experimental/MeshModelingToolsetExp/Source/MeshModelingToolsEditorOnlyExp/Private/SkeletalMesh/SkinWeightsBindingTool.cpp
+++ b/Engine/Plugins/Experimental/MeshModelingToolsetExp/Source/MeshModelingToolsEditorOnlyExp/Private/SkeletalMesh/SkinWeightsBindingTool.cpp
@@ -19,6 +19,7 @@
 #include "ModelingToolTargetUtil.h"
 #include "SkinningOps/SkinBindingOp.h"
 #include "Spatial/OccupancyGrid3.h"
+#include "ContextObjectStore.h"
 
 #include UE_INLINE_GENERATED_CPP_BY_NAME(SkinWeightsBindingTool)
 
@@ -34,17 +35,11 @@ bool USkinWeightsBindingToolBuilder::CanBuildTool(const FToolBuilderState& Scene
 
 UMultiSelectionMeshEditingTool* USkinWeightsBindingToolBuilder::CreateNewTool(const FToolBuilderState& SceneState) const
 {
-	return NewObject<USkinWeightsBindingTool>(SceneState.ToolManager);
+	USkinWeightsBindingTool* Tool = NewObject<USkinWeightsBindingTool>(SceneState.ToolManager);
+	Tool->Init(SceneState);
+	return Tool;
 }
 
-
-USkeleton* USkinWeightsBindingToolProperties::GetSkeleton(bool& bInvalidSkeletonIsError, const IPropertyHandle* PropertyHandle)
-{
-	bInvalidSkeletonIsError = false;
-	return SkeletalMesh ? SkeletalMesh->GetSkeleton() : nullptr;
-}
-
-
 USkinWeightsBindingTool::USkinWeightsBindingTool()
 {
 	Properties = CreateDefaultSubobject<USkinWeightsBindingToolProperties>(TEXT("SkinWeightsBindingProperties"));
@@ -57,6 +52,11 @@ USkinWeightsBindingTool::~USkinWeightsBindingTool()
 {
 }
 
+void USkinWeightsBindingTool::Init(const FToolBuilderState& InSceneState)
+{
+	const UContextObjectStore* ContextObjectStore = InSceneState.ToolManager->GetContextObjectStore();
+	EditorContext = ContextObjectStore->FindContext<USkeletalMeshEditorContextObjectBase>();
+}
 
 void USkinWeightsBindingTool::Setup()
 {
@@ -77,21 +77,8 @@ void USkinWeightsBindingTool::Setup()
 	if (SkeletalMeshComponent && SkeletalMeshComponent->GetSkeletalMeshAsset())
 	{
 		USkeletalMesh* SkeletalMesh = SkeletalMeshComponent->GetSkeletalMeshAsset();
-
 		ReferenceSkeleton = SkeletalMesh->GetRefSkeleton();
-		
-		// Initialize the bone browser
-		const UE::Anim::FCurveFilterSettings CurveFilterSettings = SkeletalMeshComponent->GetCurveFilterSettings(0 /* Always use the highest LOD */);
-		TArray<FBoneIndexType> AllBones;
-		AllBones.Reserve(ReferenceSkeleton.GetRawBoneNum());
-		for (int32 BoneIndex = 0; BoneIndex < ReferenceSkeleton.GetRawBoneNum(); BoneIndex++)
-		{
-			AllBones.Add(static_cast<FBoneIndexType>(BoneIndex));
-		}
-		BoneContainer.InitializeTo(AllBones, CurveFilterSettings, *SkeletalMesh);
-
-		Properties->SkeletalMesh = SkeletalMesh;
-		Properties->CurrentBone.Initialize(BoneContainer);
+		Properties->CurrentBone = ReferenceSkeleton.GetBoneName(0);
 	}
 	
 	UE::ToolTarget::HideSourceObject(Targets[0]);
@@ -128,7 +115,7 @@ void USkinWeightsBindingTool::Setup()
 		UpdateVisualization();
 	};
 
-	Properties->WatchProperty(Properties->CurrentBone.BoneName,
+	Properties->WatchProperty(Properties->CurrentBone,
 							  [HandlePropertyChange](FName) { HandlePropertyChange(false); });
 	Properties->WatchProperty(Properties->BindingType,
 							  [HandlePropertyChange](ESkinWeightsBindType) { HandlePropertyChange(true); });
@@ -169,6 +156,11 @@ void USkinWeightsBindingTool::Setup()
 
 	Preview->InvalidateResult();
 
+	if (EditorContext.IsValid())
+	{
+		EditorContext->BindTo(this);
+	}
+	
 	SetToolDisplayName(LOCTEXT("ToolName", "Bind Skin"));
 	GetToolManager()->DisplayMessage(
 		LOCTEXT("OnStartTool", "Creates a rigid binding for the skin weights."),
@@ -187,6 +179,11 @@ void USkinWeightsBindingTool::OnShutdown(EToolShutdownType ShutdownType)
 	{
 		GenerateAsset(Result);
 	}
+
+	if (EditorContext.IsValid())
+	{
+		EditorContext->UnbindFrom(this);
+	}
 }
 
 
@@ -320,7 +317,7 @@ void USkinWeightsBindingTool::UpdateVisualization(bool bInForce)
 	using namespace UE::AnimationCore;
 	using namespace UE::Geometry;
 
-	const int32 RawBoneIndex = ReferenceSkeleton.FindRawBoneIndex(Properties->CurrentBone.BoneName);
+	const int32 RawBoneIndex = ReferenceSkeleton.FindRawBoneIndex(Properties->CurrentBone);
 	if ((bInForce || Preview->HaveValidNonEmptyResult()) && RawBoneIndex != INDEX_NONE)
 	{
 		const FBoneIndexType BoneIndex = static_cast<FBoneIndexType>(RawBoneIndex);
@@ -364,6 +361,25 @@ void USkinWeightsBindingTool::UpdateVisualization(bool bInForce)
 	}
 }
 
+void USkinWeightsBindingTool::HandleSkeletalMeshModified(const TArray<FName>& InBoneNames, const ESkeletalMeshNotifyType InNotifyType)
+{
+	switch (InNotifyType)
+	{
+	case ESkeletalMeshNotifyType::BonesAdded:
+		break;
+	case ESkeletalMeshNotifyType::BonesRemoved:
+		break;
+	case ESkeletalMeshNotifyType::BonesMoved:
+		break;
+	case ESkeletalMeshNotifyType::BonesSelected:
+		{
+			Properties->CurrentBone = InBoneNames.IsEmpty() ? NAME_None : InBoneNames[0];
+		}
+		break;
+	default:
+		checkNoEntry();
+	}
+}
 
 #undef LOCTEXT_NAMESPACE
 
diff --git a/Engine/Plugins/Experimental/MeshModelingToolsetExp/Source/MeshModelingToolsEditorOnlyExp/Private/SkeletalMesh/SkinWeightsPaintTool.cpp b/Engine/Plugins/Experimental/MeshModelingToolsetExp/Source/MeshModelingToolsEditorOnlyExp/Private/SkeletalMesh/SkinWeightsPaintTool.cpp
--- a/Engine/Plugins/Experimental/MeshModelingToolsetExp/Source/MeshModelingToolsEditorOnlyExp/Private/SkeletalMesh/SkinWeightsPaintTool.cpp
+++ b/Engine/Plugins/Experimental/MeshModelingToolsetExp/Source/MeshModelingToolsEditorOnlyExp/Private/SkeletalMesh/SkinWeightsPaintTool.cpp
@@ -24,6 +24,7 @@
 #include "Spatial/PointSetHashTable.h"
 #include "Util/ColorConstants.h"
 #include "Operations/SmoothBoneWeights.h"
+#include "ContextObjectStore.h"
 
 #include UE_INLINE_GENERATED_CPP_BY_NAME(SkinWeightsPaintTool)
 
@@ -611,7 +612,15 @@ void FMeshSkinWeightsChange::AddBoneWeightEdit(const FSingleBoneWeightEdits& Bon
 
 UMeshSurfacePointTool* USkinWeightsPaintToolBuilder::CreateNewTool(const FToolBuilderState& SceneState) const
 {
-	return NewObject<USkinWeightsPaintTool>(SceneState.ToolManager);
+	USkinWeightsPaintTool* Tool = NewObject<USkinWeightsPaintTool>(SceneState.ToolManager);
+	Tool->Init(SceneState);
+	return Tool;
+}
+
+void USkinWeightsPaintTool::Init(const FToolBuilderState& InSceneState)
+{
+	const UContextObjectStore* ContextObjectStore = InSceneState.ToolManager->GetContextObjectStore();
+	EditorContext = ContextObjectStore->FindContext<USkeletalMeshEditorContextObjectBase>();
 }
 
 void USkinWeightsPaintTool::Setup()
@@ -715,6 +724,11 @@ void USkinWeightsPaintTool::Setup()
 	SmoothWeightsDataSource = MakeUnique<FPaintToolWeightsDataSource>(&Weights);
 	SmoothWeightsOp = MakeUnique<UE::Geometry::TSmoothBoneWeights<int32, float>>(PreviewMesh->GetMesh(), SmoothWeightsDataSource.Get());
 	SmoothWeightsOp->MinimumWeightThreshold = MinimumWeightThreshold;
+
+	if (EditorContext.IsValid())
+	{
+		EditorContext->BindTo(this);
+	}
 	
 	// inform user of tool keys
 	// TODO talk with UX team about viewport overlay to show hotkeys
@@ -1356,6 +1370,11 @@ void USkinWeightsPaintTool::OnShutdown(EToolShutdownType ShutdownType)
 		UE::ToolTarget::CommitMeshDescriptionUpdate(Target, EditedMesh.Get());
 		GetToolManager()->EndUndoTransaction();
 	}
+
+	if (EditorContext.IsValid())
+	{
+		EditorContext->UnbindFrom(this);
+	}
 }
 
 void USkinWeightsPaintTool::BeginChange()
diff --git a/Engine/Plugins/Experimental/MeshModelingToolsetExp/Source/MeshModelingToolsEditorOnlyExp/Public/SkeletalMesh/SkeletalMeshEditionInterface.h b/Engine/Plugins/Experimental/MeshModelingToolsetExp/Source/MeshModelingToolsEditorOnlyExp/Public/SkeletalMesh/SkeletalMeshEditionInterface.h
--- a/Engine/Plugins/Experimental/MeshModelingToolsetExp/Source/MeshModelingToolsEditorOnlyExp/Public/SkeletalMesh/SkeletalMeshEditionInterface.h
+++ b/Engine/Plugins/Experimental/MeshModelingToolsetExp/Source/MeshModelingToolsEditorOnlyExp/Public/SkeletalMesh/SkeletalMeshEditionInterface.h
@@ -134,4 +134,21 @@ public:
 
 	UPROPERTY()
 	TWeakObjectPtr<USceneComponent> Component;
+};
+
+/**
+ * USkeletalMeshEditorContextObjectBase
+ */
+
+UCLASS(Abstract)
+class MESHMODELINGTOOLSEDITORONLYEXP_API USkeletalMeshEditorContextObjectBase : public UObject
+{
+	GENERATED_BODY()
+
+public:
+	virtual void HideSkeleton() PURE_VIRTUAL(USkeletalMeshEditorContextObjectBase::HideSkeleton, return;)
+	virtual void ShowSkeleton() PURE_VIRTUAL(USkeletalMeshEditorContextObjectBase::ShowSkeleton, return;)
+	
+	virtual void BindTo(ISkeletalMeshEditingInterface* InEditingInterface) PURE_VIRTUAL(USkeletalMeshEditorContextObjectBase::BindTo, return;)
+	virtual void UnbindFrom(ISkeletalMeshEditingInterface* InEditingInterface) PURE_VIRTUAL(USkeletalMeshEditorContextObjectBase::UnbindFrom, return;)
 };
\ No newline at end of file
diff --git a/Engine/Plugins/Experimental/MeshModelingToolsetExp/Source/MeshModelingToolsEditorOnlyExp/Public/SkeletalMesh/SkeletonEditingTool.h b/Engine/Plugins/Experimental/MeshModelingToolsetExp/Source/MeshModelingToolsEditorOnlyExp/Public/SkeletalMesh/SkeletonEditingTool.h
--- a/Engine/Plugins/Experimental/MeshModelingToolsetExp/Source/MeshModelingToolsEditorOnlyExp/Public/SkeletalMesh/SkeletonEditingTool.h
+++ b/Engine/Plugins/Experimental/MeshModelingToolsetExp/Source/MeshModelingToolsEditorOnlyExp/Public/SkeletalMesh/SkeletonEditingTool.h
@@ -195,6 +195,9 @@ protected:
 	UPROPERTY()
 	TObjectPtr<USkeletalMeshGizmoWrapperBase> GizmoWrapper = nullptr;
 
+	UPROPERTY()
+	TWeakObjectPtr<USkeletalMeshEditorContextObjectBase> EditorContext = nullptr;
+
 	void UpdateGizmo() const;
 	
 	// ref skeleton transactions
diff --git a/Engine/Plugins/Experimental/MeshModelingToolsetExp/Source/MeshModelingToolsEditorOnlyExp/Public/SkeletalMesh/SkinWeightsBindingTool.h b/Engine/Plugins/Experimental/MeshModelingToolsetExp/Source/MeshModelingToolsEditorOnlyExp/Public/SkeletalMesh/SkinWeightsBindingTool.h
--- a/Engine/Plugins/Experimental/MeshModelingToolsetExp/Source/MeshModelingToolsEditorOnlyExp/Public/SkeletalMesh/SkinWeightsBindingTool.h
+++ b/Engine/Plugins/Experimental/MeshModelingToolsetExp/Source/MeshModelingToolsEditorOnlyExp/Public/SkeletalMesh/SkinWeightsBindingTool.h
@@ -7,6 +7,7 @@
 #include "BoneContainer.h"
 #include "InteractiveToolBuilder.h"
 #include "ModelingOperators.h"
+#include "SkeletalMeshEditionInterface.h"
 #include "BaseTools/MultiSelectionMeshEditingTool.h"
 #include "Interfaces/Interface_BoneReferenceSkeletonProvider.h"
 
@@ -45,13 +46,12 @@ enum class ESkinWeightsBindType : uint8
 
 UCLASS()
 class MESHMODELINGTOOLSEDITORONLYEXP_API USkinWeightsBindingToolProperties : 
-	public UInteractiveToolPropertySet,
-	public IBoneReferenceSkeletonProvider
+	public UInteractiveToolPropertySet
 {
 	GENERATED_BODY()
 public:
-	UPROPERTY(VisibleAnywhere, Category = Skeleton)
-	FBoneReference CurrentBone;
+	UPROPERTY()
+	FName CurrentBone = NAME_None;
 
 	/** Binding type to use */
 	UPROPERTY(EditAnywhere, Category = Binding)
@@ -71,11 +71,6 @@ public:
 	
 	// UPROPERTY(EditAnywhere, Category = Debug)
 	bool bDebugDraw = false;
-	
-	// IBoneReferenceSkeletonProvider
-	virtual USkeleton* GetSkeleton(bool& bInvalidSkeletonIsError, const IPropertyHandle* PropertyHandle) override;
-
-	TObjectPtr<USkeletalMesh> SkeletalMesh;
 };
 
 /**
@@ -84,13 +79,16 @@ public:
 UCLASS()
 class MESHMODELINGTOOLSEDITORONLYEXP_API USkinWeightsBindingTool :
 	public UMultiSelectionMeshEditingTool,
-	public UE::Geometry::IDynamicMeshOperatorFactory
+	public UE::Geometry::IDynamicMeshOperatorFactory,
+	public ISkeletalMeshEditingInterface
 {
 	GENERATED_BODY()
 
 public:
 	USkinWeightsBindingTool();
 	virtual ~USkinWeightsBindingTool() override;
+
+	void Init(const FToolBuilderState& InSceneState);
 	
 	virtual void Setup() override;
 	virtual void OnShutdown(EToolShutdownType ShutdownType) override;
@@ -117,10 +115,15 @@ protected:
 	TSharedPtr<UE::Geometry::FDynamicMesh3, ESPMode::ThreadSafe> OriginalMesh;
 
 	FReferenceSkeleton ReferenceSkeleton;
-	FBoneContainer BoneContainer;
 	
 	void GenerateAsset(const FDynamicMeshOpResult& Result);
 
 	static FVector4f WeightToColor(float Value);
 	void UpdateVisualization(bool bInForce = false);
+
+	UPROPERTY()
+	TWeakObjectPtr<USkeletalMeshEditorContextObjectBase> EditorContext = nullptr;
+
+	// ISkeletalMeshEditionInterface
+	virtual void HandleSkeletalMeshModified(const TArray<FName>& InBoneNames, const ESkeletalMeshNotifyType InNotifyType) override;
 };
diff --git a/Engine/Plugins/Experimental/MeshModelingToolsetExp/Source/MeshModelingToolsEditorOnlyExp/Public/SkeletalMesh/SkinWeightsPaintTool.h b/Engine/Plugins/Experimental/MeshModelingToolsetExp/Source/MeshModelingToolsEditorOnlyExp/Public/SkeletalMesh/SkinWeightsPaintTool.h
--- a/Engine/Plugins/Experimental/MeshModelingToolsetExp/Source/MeshModelingToolsEditorOnlyExp/Public/SkeletalMesh/SkinWeightsPaintTool.h
+++ b/Engine/Plugins/Experimental/MeshModelingToolsetExp/Source/MeshModelingToolsEditorOnlyExp/Public/SkeletalMesh/SkinWeightsPaintTool.h
@@ -302,6 +302,8 @@ public:
 	virtual void OnUpdateDrag(const FRay& Ray) override;
 	virtual void OnEndDrag(const FRay& Ray) override;
 	virtual bool OnUpdateHover(const FInputDeviceRay& DevicePos) override;
+
+	void Init(const FToolBuilderState& InSceneState);
 	
 	// UInteractiveTool
 	virtual void Setup() override;
@@ -414,5 +416,8 @@ protected:
 	TUniquePtr<UE::Geometry::FTriangleGroupTopology> SelectionTopology = nullptr;
 	void OnSelectionModified();
 
+	UPROPERTY()
+	TWeakObjectPtr<USkeletalMeshEditorContextObjectBase> EditorContext = nullptr;
+
 	friend SkinPaintTool::FSkinToolDeformer;
 };
