Swapping Skeletal Mesh Editing tools to use toolkit builder UI.
New default tab layout for skeletal mesh editor.
Skel Mesh tools activate immediately when opening editor.

diff --git a/Engine/Plugins/Editor/ModelingToolsEditorMode/Source/ModelingToolsEditorMode/Private/ModelingToolsEditorModeStyle.cpp b/Engine/Plugins/Editor/ModelingToolsEditorMode/Source/ModelingToolsEditorMode/Private/ModelingToolsEditorModeStyle.cpp
--- a/Engine/Plugins/Editor/ModelingToolsEditorMode/Source/ModelingToolsEditorMode/Private/ModelingToolsEditorModeStyle.cpp
+++ b/Engine/Plugins/Editor/ModelingToolsEditorMode/Source/ModelingToolsEditorMode/Private/ModelingToolsEditorModeStyle.cpp
@@ -122,6 +122,9 @@ void FModelingToolsEditorModeStyle::Initialize()
 		StyleSet->Set("ModelingToolsManagerCommands.LoadBakingTools", new IMAGE_BRUSH_SVG( "Icons/LoadBakingTools", Icon20x20 ) );
 		StyleSet->Set("ModelingToolsManagerCommands.LoadVolumeTools", new IMAGE_BRUSH_SVG( "Icons/LoadVolumeTools", Icon20x20 ) );
 		StyleSet->Set("ModelingToolsManagerCommands.LoadLodsTools", new IMAGE_BRUSH_SVG( "Icons/LoadLodsTools", Icon20x20 ) );
+		
+		StyleSet->Set("ModelingToolsManagerCommands.LoadSkinTools", new IMAGE_BRUSH_SVG( "Icons/LoadDeformTools", Icon20x20 ) );
+		StyleSet->Set("ModelingToolsManagerCommands.LoadSkeletonTools", new IMAGE_BRUSH_SVG( "Icons/LoadTransformTools", Icon20x20 ) );
 
 		StyleSet->Set("ModelingToolsManagerCommands.BeginShapeSprayTool", 				new IMAGE_PLUGIN_BRUSH("Icons/ShapeSpray_40x",	Icon20x20));
 		StyleSet->Set("ModelingToolsManagerCommands.BeginShapeSprayTool.Small", 		new IMAGE_PLUGIN_BRUSH("Icons/ShapeSpray_40x",	Icon20x20));
@@ -335,7 +338,9 @@ void FModelingToolsEditorModeStyle::Initialize()
 		StyleSet->Set("ModelingToolsManagerCommands.BeginGenerateLODMeshesTool.Small", new IMAGE_BRUSH_SVG("Icons/GenLODs", Icon20x20));
 		StyleSet->Set("ModelingToolsManagerCommands.BeginGroomToMeshTool", new IMAGE_BRUSH_SVG("Icons/HairHelmet", Icon20x20));
 		StyleSet->Set("ModelingToolsManagerCommands.BeginGroomToMeshTool.Small", new IMAGE_BRUSH_SVG("Icons/HairHelmet", Icon20x20));
-
+		
+		StyleSet->Set("ModelingToolsManagerCommands.BeginSkeletonEditingTool", new IMAGE_BRUSH_SVG("Icons/EditSkeleton", Icon20x20));
+		StyleSet->Set("ModelingToolsManagerCommands.BeginSkinWeightsBindingTool", new IMAGE_BRUSH_SVG("Icons/BindSkin", Icon20x20));
 		StyleSet->Set("ModelingToolsManagerCommands.BeginSkinWeightsPaintTool", new IMAGE_BRUSH_SVG("Icons/SkinWeightsPaint", Icon20x20));
 		StyleSet->Set("ModelingToolsManagerCommands.BeginSkinWeightsPaintTool.Small", new IMAGE_BRUSH_SVG("Icons/SkinWeightsPaint", Icon20x20));
 
diff --git a/Engine/Plugins/Editor/ModelingToolsEditorMode/Source/ModelingToolsEditorMode/Private/ModelingToolsManagerActions.cpp b/Engine/Plugins/Editor/ModelingToolsEditorMode/Source/ModelingToolsEditorMode/Private/ModelingToolsManagerActions.cpp
--- a/Engine/Plugins/Editor/ModelingToolsEditorMode/Source/ModelingToolsEditorMode/Private/ModelingToolsManagerActions.cpp
+++ b/Engine/Plugins/Editor/ModelingToolsEditorMode/Source/ModelingToolsEditorMode/Private/ModelingToolsManagerActions.cpp
@@ -324,10 +324,10 @@ void FModelingToolsManagerCommands::RegisterCommands()
 	REGISTER_MODELING_TOOL_COMMAND(BeginShapeSprayTool, "Spray", "Shape Spray");
 	REGISTER_MODELING_TOOL_COMMAND(BeginEditCollisionGeometryTool, "EditPhys", "Edit Simple Collision Geometry for selected Mesh");
 	
-	REGISTER_MODELING_TOOL_COMMAND(BeginSkinWeightsPaintTool, "Weights", "Start the Skin Weights Editing Tool");
-	REGISTER_MODELING_TOOL_COMMAND(BeginSkinWeightsBindingTool, "Bind", "Start the Skin Weights Binding Tool");
+	REGISTER_MODELING_TOOL_COMMAND(BeginSkinWeightsPaintTool, "Edit Weights", "Tune the per-vertex skin weights.");
+	REGISTER_MODELING_TOOL_COMMAND(BeginSkinWeightsBindingTool, "Bind Skin", "Create default weights by binding the skin to bones.");
 
-	REGISTER_MODELING_TOOL_COMMAND(BeginSkeletonEditingTool, "SkeletonEditing", "Start the Skeleton Editing Tool");
+	REGISTER_MODELING_TOOL_COMMAND(BeginSkeletonEditingTool, "Edit Skeleton", "Add, Remove, Reparent, Move and Rename bones.");
 
 	REGISTER_MODELING_TOOL_COMMAND(BeginPolyModelTool_Inset, "Inset", "Inset the current Selected Faces");
 	REGISTER_MODELING_TOOL_COMMAND(BeginPolyModelTool_Outset, "Outset", "Outset the current Selected Faces");
@@ -403,6 +403,8 @@ void FModelingToolsManagerCommands::RegisterCommands()
 	UI_COMMAND(LoadBakingTools, "Bake", "Bake Textures & Colors", EUserInterfaceActionType::RadioButton, FInputChord());
 	UI_COMMAND(LoadVolumeTools, "Volumes", "Volumes", EUserInterfaceActionType::RadioButton, FInputChord());
 	UI_COMMAND(LoadLodsTools, "Misc", "Additional Utility Tools", EUserInterfaceActionType::RadioButton, FInputChord());
+	UI_COMMAND(LoadSkinTools, "Skin", "Edit Skin Weights", EUserInterfaceActionType::RadioButton, FInputChord());
+	UI_COMMAND(LoadSkeletonTools, "Skeleton", "Edit Bones", EUserInterfaceActionType::RadioButton, FInputChord());
 
 	UI_COMMAND(AcceptActiveTool, "Accept", "Accept the active Tool", EUserInterfaceActionType::Button, FInputChord());
 	UI_COMMAND(CancelActiveTool, "Cancel", "Cancel the active Tool", EUserInterfaceActionType::Button, FInputChord());
diff --git a/Engine/Plugins/Editor/ModelingToolsEditorMode/Source/ModelingToolsEditorMode/Public/ModelingToolsManagerActions.h b/Engine/Plugins/Editor/ModelingToolsEditorMode/Source/ModelingToolsEditorMode/Public/ModelingToolsManagerActions.h
--- a/Engine/Plugins/Editor/ModelingToolsEditorMode/Source/ModelingToolsEditorMode/Public/ModelingToolsManagerActions.h
+++ b/Engine/Plugins/Editor/ModelingToolsEditorMode/Source/ModelingToolsEditorMode/Public/ModelingToolsManagerActions.h
@@ -198,6 +198,9 @@ public:
 	TSharedPtr<FUICommandInfo> LoadLodsTools;
 	// For connecting ModelingMode with UVEditor plugin
 	TSharedPtr<FUICommandInfo> LaunchUVEditor;
+	// skeletal mesh editing tools
+	TSharedPtr<FUICommandInfo> LoadSkinTools;
+	TSharedPtr<FUICommandInfo> LoadSkeletonTools;
 
 	//
 	// Accept/Cancel/Complete commands are used to end the active Tool via ToolManager
diff --git a/Engine/Plugins/Experimental/Animation/SkeletalMeshModelingTools/Source/SkeletalMeshModelingTools/Private/SkeletalMeshModelingToolsEditorMode.cpp b/Engine/Plugins/Experimental/Animation/SkeletalMeshModelingTools/Source/SkeletalMeshModelingTools/Private/SkeletalMeshModelingToolsEditorMode.cpp
--- a/Engine/Plugins/Experimental/Animation/SkeletalMeshModelingTools/Source/SkeletalMeshModelingTools/Private/SkeletalMeshModelingToolsEditorMode.cpp
+++ b/Engine/Plugins/Experimental/Animation/SkeletalMeshModelingTools/Source/SkeletalMeshModelingTools/Private/SkeletalMeshModelingToolsEditorMode.cpp
@@ -155,7 +155,7 @@ const FEditorModeID USkeletalMeshModelingToolsEditorMode::Id("SkeletalMeshModeli
 
 USkeletalMeshModelingToolsEditorMode::USkeletalMeshModelingToolsEditorMode() 
 {
-	Info = FEditorModeInfo(Id, LOCTEXT("SkeletalMeshModelingMode", "Skeletal Mesh Modeling"), FSlateIcon(), false);
+	Info = FEditorModeInfo(Id, LOCTEXT("SkeletalMeshEditingMode", "Skeletal Mesh Editing"), FSlateIcon(), false);
 }
 
 
diff --git a/Engine/Plugins/Experimental/Animation/SkeletalMeshModelingTools/Source/SkeletalMeshModelingTools/Private/SkeletalMeshModelingToolsEditorModeToolkit.cpp b/Engine/Plugins/Experimental/Animation/SkeletalMeshModelingTools/Source/SkeletalMeshModelingTools/Private/SkeletalMeshModelingToolsEditorModeToolkit.cpp
--- a/Engine/Plugins/Experimental/Animation/SkeletalMeshModelingTools/Source/SkeletalMeshModelingTools/Private/SkeletalMeshModelingToolsEditorModeToolkit.cpp
+++ b/Engine/Plugins/Experimental/Animation/SkeletalMeshModelingTools/Source/SkeletalMeshModelingTools/Private/SkeletalMeshModelingToolsEditorModeToolkit.cpp
@@ -36,22 +36,11 @@ FSkeletalMeshModelingToolsEditorModeToolkit::~FSkeletalMeshModelingToolsEditorMo
 
 void FSkeletalMeshModelingToolsEditorModeToolkit::Init(
 	const TSharedPtr<IToolkitHost>& InToolkitHost, 
-	TWeakObjectPtr<UEdMode> InOwningMode
-	)
+	TWeakObjectPtr<UEdMode> InOwningMode)
 {
-	// Create a details view to show the tool props
-	FPropertyEditorModule& PropertyEditorModule = FModuleManager::GetModuleChecked<FPropertyEditorModule>("PropertyEditor");
+	bUsesToolkitBuilder = true;
 
-	FDetailsViewArgs DetailsViewArgs;
-
-	DetailsViewArgs.bHideSelectionTip = true;
-	DetailsViewArgs.NameAreaSettings = FDetailsViewArgs::HideNameArea;
-	DetailsViewArgs.bAllowSearch = false;
-	DetailsViewArgs.DefaultsOnlyVisibility = EEditDefaultsOnlyNodeVisibility::Automatic;
-	DetailsViewArgs.bShowOptions = false;
-	DetailsViewArgs.bAllowMultipleTopLevelObjects = true;
-
-	DetailsView = PropertyEditorModule.CreateDetailView(DetailsViewArgs);
+	FModeToolkit::Init(InToolkitHost, InOwningMode);
 
 	ModeWarningArea = SNew(STextBlock)
 		.AutoWrapText(true)
@@ -65,8 +54,7 @@ void FSkeletalMeshModelingToolsEditorModeToolkit::Init(
 		.Font(FCoreStyle::GetDefaultFontStyle("Bold", 12));
 	ModeHeaderArea->SetText(LOCTEXT("SelectToolLabel", "Select a Tool from the Toolbar"));
 	ModeHeaderArea->SetJustification(ETextJustify::Center);
-
-
+	
 	ToolWarningArea = SNew(STextBlock)
 		.AutoWrapText(true)
 		.Font(FCoreStyle::GetDefaultFontStyle("Bold", 9))
@@ -74,34 +62,26 @@ void FSkeletalMeshModelingToolsEditorModeToolkit::Init(
 	ToolWarningArea->SetText(FText::GetEmpty());
 
 
-	SAssignNew(ToolkitWidget, SBox)
-		.HAlign(HAlign_Fill)
-		.Padding(2)
-		[
-			SNew(SVerticalBox)
-			+ SVerticalBox::Slot().AutoHeight().HAlign(HAlign_Fill).Padding(5)
-			[
-				ModeWarningArea->AsShared()
-			]
-			+ SVerticalBox::Slot().AutoHeight().HAlign(HAlign_Fill).Padding(5)
-			[
-				ModeHeaderArea->AsShared()
-			]
-			+ SVerticalBox::Slot().AutoHeight().HAlign(HAlign_Fill).Padding(5)
-			[
-				ToolWarningArea->AsShared()
-			]
-			+ SVerticalBox::Slot().FillHeight(1.f).HAlign(HAlign_Fill)
-			[
-				DetailsView->AsShared()
-			]
-		];
-
-	FModeToolkit::Init(InToolkitHost, InOwningMode);
+	RegisterPalettes();
 
 	ClearNotification();
 	ClearWarning();
+	
+	// create the toolkit widget
+	{
+		ToolkitSections->ModeWarningArea = ModeWarningArea;
+		ToolkitSections->DetailsView = ModeDetailsView;
+		ToolkitSections->ToolWarningArea = ToolWarningArea;
 
+		SAssignNew(ToolkitWidget, SBorder)
+		.HAlign(HAlign_Fill)
+		.Padding(0)
+		.BorderImage(FAppStyle::GetBrush("ToolPanel.GroupBorder"))
+		[
+			ToolkitBuilder->GenerateWidget()->AsShared()
+		];	
+	}
+	
 	ActiveToolName = FText::GetEmpty();
 	ActiveToolMessage = FText::GetEmpty();
 
@@ -111,6 +91,137 @@ void FSkeletalMeshModelingToolsEditorModeToolkit::Init(
 		Context->OnToolWarningMessage.AddSP(this, &FSkeletalMeshModelingToolsEditorModeToolkit::PostWarning);
 	}
 
+	// add viewport overlay widget to accept / cancel tool
+	MakeToolAcceptCancelWidget();
+}
+
+FName FSkeletalMeshModelingToolsEditorModeToolkit::GetToolkitFName() const
+{
+	return FName("SkeletalMeshModelingToolsEditorModeToolkit");
+}
+
+FText FSkeletalMeshModelingToolsEditorModeToolkit::GetBaseToolkitName() const
+{
+	return LOCTEXT("DisplayName", "Skeletal Mesh Editing Tools");
+}
+
+void FSkeletalMeshModelingToolsEditorModeToolkit::OnToolStarted(UInteractiveToolManager* Manager, UInteractiveTool* Tool)
+{
+	UpdateActiveToolProperties(Tool);
+
+	Tool->OnPropertySetsModified.AddSP(this, &FSkeletalMeshModelingToolsEditorModeToolkit::UpdateActiveToolProperties, Tool);
+
+	ModeHeaderArea->SetVisibility(EVisibility::Collapsed);
+	ActiveToolName = Tool->GetToolInfo().ToolDisplayName;
+
+	FString ActiveToolIdentifier = GetScriptableEditorMode()->GetToolManager()->GetActiveToolName(EToolSide::Left);
+	ActiveToolIdentifier.InsertAt(0, ".");
+	FName ActiveToolIconName = ISlateStyle::Join(FModelingToolsManagerCommands::Get().GetContextName(), TCHAR_TO_ANSI(*ActiveToolIdentifier));
+	ActiveToolIcon = FModelingToolsEditorModeStyle::Get()->GetOptionalBrush(ActiveToolIconName);
+	
+	GetToolkitHost()->AddViewportOverlayWidget(ViewportOverlayWidget.ToSharedRef());
+}
+
+void FSkeletalMeshModelingToolsEditorModeToolkit::OnToolEnded(UInteractiveToolManager* Manager, UInteractiveTool* Tool)
+{
+	if (IsHosted())
+	{
+		GetToolkitHost()->RemoveViewportOverlayWidget(ViewportOverlayWidget.ToSharedRef());
+	}
+	
+	if (Tool)
+	{
+		Tool->OnPropertySetsModified.RemoveAll(this);
+	}
+
+	ModeHeaderArea->SetVisibility(EVisibility::Visible);
+	ModeDetailsView->SetObject(nullptr);
+	ActiveToolName = FText::GetEmpty();
+	ClearNotification();
+	ClearWarning();
+}
+
+FText FSkeletalMeshModelingToolsEditorModeToolkit::GetActiveToolDisplayName() const
+{
+	return ActiveToolName;
+}
+
+
+FText FSkeletalMeshModelingToolsEditorModeToolkit::GetActiveToolMessage() const
+{
+	return ActiveToolMessage;
+}
+
+void FSkeletalMeshModelingToolsEditorModeToolkit::RegisterPalettes()
+{
+	const FModelingToolsManagerCommands& Commands = FModelingToolsManagerCommands::Get();
+	const TSharedPtr<FUICommandList> CommandList = GetToolkitCommands();
+	
+	ToolkitSections = MakeShared<FToolkitSections>();
+	FToolkitBuilderArgs ToolkitBuilderArgs(GetScriptableEditorMode()->GetModeInfo().ToolbarCustomizationName);
+	ToolkitBuilderArgs.ToolkitCommandList = GetToolkitCommands();
+	ToolkitBuilderArgs.ToolkitSections = ToolkitSections;
+	ToolkitBuilderArgs.SelectedCategoryTitleVisibility = EVisibility::Collapsed;
+	ToolkitBuilder = MakeShared<FToolkitBuilder>(ToolkitBuilderArgs);
+
+	const TArray<TSharedPtr<FUICommandInfo>> SkeletonCommands({
+		Commands.BeginSkeletonEditingTool
+	});
+	ToolkitBuilder->AddPalette(
+		MakeShareable( new FToolPalette( Commands.LoadSkeletonTools.ToSharedRef(), SkeletonCommands ) ) );
+
+	const TArray<TSharedPtr<FUICommandInfo>> SkinCommands({
+		Commands.BeginSkinWeightsBindingTool,
+		Commands.BeginSkinWeightsPaintTool,
+		Commands.BeginAttributeEditorTool,
+		Commands.BeginMeshAttributePaintTool
+	});
+	ToolkitBuilder->AddPalette(
+		MakeShareable( new FToolPalette( Commands.LoadSkinTools.ToSharedRef(), SkinCommands ) ) );
+
+	const TArray<TSharedPtr<FUICommandInfo>> DeformCommands({
+		Commands.BeginSculptMeshTool,
+		Commands.BeginRemeshSculptMeshTool,
+		Commands.BeginSmoothMeshTool,
+		Commands.BeginOffsetMeshTool,
+		Commands.BeginMeshSpaceDeformerTool,
+		Commands.BeginLatticeDeformerTool,
+		Commands.BeginDisplaceMeshTool
+	});
+	ToolkitBuilder->AddPalette(
+		MakeShareable( new FToolPalette( Commands.LoadDeformTools.ToSharedRef(), DeformCommands ) ) );
+
+	const TArray<TSharedPtr<FUICommandInfo>> ModelingMeshCommands({
+		Commands.BeginPolyEditTool,
+		Commands.BeginPolyDeformTool,
+		Commands.BeginHoleFillTool,
+		Commands.BeginPolygonCutTool,
+		});
+	ToolkitBuilder->AddPalette(
+		MakeShareable( new FToolPalette( Commands.LoadPolyTools.ToSharedRef(), ModelingMeshCommands ) ) );
+
+	const TArray<TSharedPtr<FUICommandInfo>> ProcessMeshCommands({
+		Commands.BeginSimplifyMeshTool,
+		Commands.BeginRemeshMeshTool,
+		Commands.BeginWeldEdgesTool,
+		Commands.BeginRemoveOccludedTrianglesTool,
+		Commands.BeginProjectToTargetTool
+	});
+	ToolkitBuilder->AddPalette(
+		MakeShareable( new FToolPalette( Commands.LoadMeshOpsTools.ToSharedRef(), ProcessMeshCommands ) ) );
+
+	ToolkitBuilder->SetActivePaletteOnLoad(Commands.BeginSkeletonEditingTool.Get());
+	ToolkitBuilder->UpdateWidget();
+	
+	// if selected palette changes, make sure we are showing the palette command buttons, which may be hidden by active Tool
+	ActivePaletteChangedHandle = ToolkitBuilder->OnActivePaletteChanged.AddLambda([this]()
+	{
+		ToolkitBuilder->SetActivePaletteCommandsVisibility(EVisibility::Visible);
+	});
+}
+
+void FSkeletalMeshModelingToolsEditorModeToolkit::MakeToolAcceptCancelWidget()
+{
 	SAssignNew(ViewportOverlayWidget, SHorizontalBox)
 
 	+SHorizontalBox::Slot()
@@ -183,136 +294,6 @@ void FSkeletalMeshModelingToolsEditorModeToolkit::Init(
 	];
 }
 
-
-FName FSkeletalMeshModelingToolsEditorModeToolkit::GetToolkitFName() const
-{
-	return FName("SkeletalMeshModelingToolsEditorModeToolkit");
-}
-
-
-FText FSkeletalMeshModelingToolsEditorModeToolkit::GetBaseToolkitName() const
-{
-	return LOCTEXT("DisplayName", "Skeletal Mesh Modeling Tools");
-}
-
-
-void FSkeletalMeshModelingToolsEditorModeToolkit::OnToolStarted(UInteractiveToolManager* Manager, UInteractiveTool* Tool)
-{
-	UpdateActiveToolProperties(Tool);
-
-	Tool->OnPropertySetsModified.AddSP(this, &FSkeletalMeshModelingToolsEditorModeToolkit::UpdateActiveToolProperties, Tool);
-
-	ModeHeaderArea->SetVisibility(EVisibility::Collapsed);
-	ActiveToolName = Tool->GetToolInfo().ToolDisplayName;
-
-	FString ActiveToolIdentifier = GetScriptableEditorMode()->GetToolManager()->GetActiveToolName(EToolSide::Left);
-	ActiveToolIdentifier.InsertAt(0, ".");
-	FName ActiveToolIconName = ISlateStyle::Join(FModelingToolsManagerCommands::Get().GetContextName(), TCHAR_TO_ANSI(*ActiveToolIdentifier));
-	ActiveToolIcon = FModelingToolsEditorModeStyle::Get()->GetOptionalBrush(ActiveToolIconName);
-	
-	GetToolkitHost()->AddViewportOverlayWidget(ViewportOverlayWidget.ToSharedRef());
-}
-
-
-void FSkeletalMeshModelingToolsEditorModeToolkit::OnToolEnded(UInteractiveToolManager* Manager, UInteractiveTool* Tool)
-{
-	if (IsHosted())
-	{
-		GetToolkitHost()->RemoveViewportOverlayWidget(ViewportOverlayWidget.ToSharedRef());
-	}
-	
-	if (Tool)
-	{
-		Tool->OnPropertySetsModified.RemoveAll(this);
-	}
-
-	ModeHeaderArea->SetVisibility(EVisibility::Visible);
-	DetailsView->SetObject(nullptr);
-	ActiveToolName = FText::GetEmpty();
-	ClearNotification();
-	ClearWarning();
-
-}
-
-
-FText FSkeletalMeshModelingToolsEditorModeToolkit::GetActiveToolDisplayName() const
-{
-	return ActiveToolName;
-}
-
-
-FText FSkeletalMeshModelingToolsEditorModeToolkit::GetActiveToolMessage() const
-{
-	return ActiveToolMessage;
-}
-
-static const FName EditTabName(TEXT("Edit Mesh"));
-static const FName ProcessingTabName(TEXT("MeshOps"));
-static const FName DeformTabName(TEXT("Deform"));
-static const FName SkinWeightsTabName(TEXT("Skin"));
-static const FName SkeletonTabName(TEXT("Skeleton"));
-
-
-void FSkeletalMeshModelingToolsEditorModeToolkit::GetToolPaletteNames(TArray<FName>& InPaletteName) const
-{
-	InPaletteName = { EditTabName, ProcessingTabName, DeformTabName, SkinWeightsTabName, SkeletonTabName };
-}
-
-
-FText FSkeletalMeshModelingToolsEditorModeToolkit::GetToolPaletteDisplayName(FName PaletteName) const
-{
-	return FText::FromName(PaletteName);
-}
-
-
-void FSkeletalMeshModelingToolsEditorModeToolkit::BuildToolPalette(FName PaletteName, class FToolBarBuilder& ToolbarBuilder)
-{
-	const FModelingToolsManagerCommands& Commands = FModelingToolsManagerCommands::Get();
-
-	if (PaletteName == EditTabName)
-	{
-		ToolbarBuilder.AddToolBarButton(Commands.BeginPolyEditTool);
-		ToolbarBuilder.AddToolBarButton(Commands.BeginPolyDeformTool);
-		ToolbarBuilder.AddToolBarButton(Commands.BeginHoleFillTool);
-		ToolbarBuilder.AddToolBarButton(Commands.BeginPolygonCutTool);
-	}
-	else if (PaletteName == ProcessingTabName)
-	{
-		ToolbarBuilder.AddToolBarButton(Commands.BeginSimplifyMeshTool);
-		ToolbarBuilder.AddToolBarButton(Commands.BeginRemeshMeshTool);
-		ToolbarBuilder.AddToolBarButton(Commands.BeginWeldEdgesTool);
-		ToolbarBuilder.AddToolBarButton(Commands.BeginRemoveOccludedTrianglesTool);
-		ToolbarBuilder.AddToolBarButton(Commands.BeginProjectToTargetTool);
-	}
-	else if (PaletteName == DeformTabName)
-	{
-		ToolbarBuilder.AddToolBarButton(Commands.BeginSculptMeshTool);
-		ToolbarBuilder.AddToolBarButton(Commands.BeginRemeshSculptMeshTool);
-		ToolbarBuilder.AddToolBarButton(Commands.BeginSmoothMeshTool);
-		ToolbarBuilder.AddToolBarButton(Commands.BeginOffsetMeshTool);
-		ToolbarBuilder.AddToolBarButton(Commands.BeginMeshSpaceDeformerTool);
-		ToolbarBuilder.AddToolBarButton(Commands.BeginLatticeDeformerTool);
-		ToolbarBuilder.AddToolBarButton(Commands.BeginDisplaceMeshTool);
-	}
-	else if (PaletteName == SkinWeightsTabName)
-	{
-		ToolbarBuilder.AddToolBarButton(Commands.BeginAttributeEditorTool);
-		ToolbarBuilder.AddToolBarButton(Commands.BeginMeshAttributePaintTool);
-		ToolbarBuilder.AddToolBarButton(Commands.BeginSkinWeightsPaintTool);
-		ToolbarBuilder.AddToolBarButton(Commands.BeginSkinWeightsBindingTool);
-	}
-	else if (PaletteName == SkeletonTabName)
-	{
-		ToolbarBuilder.AddToolBarButton(Commands.BeginSkeletonEditingTool);
-	}
-}
-
-
-void FSkeletalMeshModelingToolsEditorModeToolkit::OnToolPaletteChanged(FName PaletteName)
-{
-}
-
-
 void FSkeletalMeshModelingToolsEditorModeToolkit::PostNotification(const FText& InMessage)
 {
 	ClearNotification();
@@ -356,11 +337,13 @@ void FSkeletalMeshModelingToolsEditorModeToolkit::ClearWarning()
 
 void FSkeletalMeshModelingToolsEditorModeToolkit::UpdateActiveToolProperties(UInteractiveTool* Tool)
 {
-	if (Tool)
+	UInteractiveTool* CurTool = GetScriptableEditorMode()->GetToolManager(EToolsContextScope::EdMode)->GetActiveTool(EToolSide::Left);
+	if (CurTool == nullptr)
 	{
-		DetailsView->SetObjects(Tool->GetToolProperties(true));
+		return;
 	}
+		
+	ModeDetailsView->SetObjects(CurTool->GetToolProperties(true));
 }
 
-
 #undef LOCTEXT_NAMESPACE
diff --git a/Engine/Plugins/Experimental/Animation/SkeletalMeshModelingTools/Source/SkeletalMeshModelingTools/Private/SkeletalMeshModelingToolsEditorModeToolkit.h b/Engine/Plugins/Experimental/Animation/SkeletalMeshModelingTools/Source/SkeletalMeshModelingTools/Private/SkeletalMeshModelingToolsEditorModeToolkit.h
--- a/Engine/Plugins/Experimental/Animation/SkeletalMeshModelingTools/Source/SkeletalMeshModelingTools/Private/SkeletalMeshModelingToolsEditorModeToolkit.h
+++ b/Engine/Plugins/Experimental/Animation/SkeletalMeshModelingTools/Source/SkeletalMeshModelingTools/Private/SkeletalMeshModelingToolsEditorModeToolkit.h
@@ -12,30 +12,23 @@ class STextBlock;
 class SButton;
 
 
-class FSkeletalMeshModelingToolsEditorModeToolkit : 
-	 public FModeToolkit
+class FSkeletalMeshModelingToolsEditorModeToolkit : public FModeToolkit
 {
 public:
 
-	~FSkeletalMeshModelingToolsEditorModeToolkit();
+	virtual ~FSkeletalMeshModelingToolsEditorModeToolkit() override;
 	
 	// IToolkit overrides
-	void Init(const TSharedPtr<IToolkitHost>& InToolkitHost, TWeakObjectPtr<UEdMode> InOwningMode) override;
-	FName GetToolkitFName() const override;
-	FText GetBaseToolkitName() const override;
-	TSharedPtr<SWidget> GetInlineContent() const override { return ToolkitWidget; }
-
-	void OnToolStarted(UInteractiveToolManager* Manager, UInteractiveTool* Tool) override;
-	void OnToolEnded(UInteractiveToolManager* Manager, UInteractiveTool* Tool) override;
-
-	FText GetActiveToolDisplayName() const override;
-	FText GetActiveToolMessage() const override;
-
-	void GetToolPaletteNames(TArray<FName>& InPaletteName) const override;
-	FText GetToolPaletteDisplayName(FName PaletteName) const override;
-	void BuildToolPalette(FName PaletteName, class FToolBarBuilder& ToolbarBuilder) override;
-	void OnToolPaletteChanged(FName PaletteName) override;
+	virtual void Init(const TSharedPtr<IToolkitHost>& InToolkitHost, TWeakObjectPtr<UEdMode> InOwningMode) override;
+	virtual FName GetToolkitFName() const override;
+	virtual FText GetBaseToolkitName() const override;
+	virtual TSharedPtr<SWidget> GetInlineContent() const override { return ToolkitWidget; }
 
+	// FModeToolkit overrides
+	virtual void OnToolStarted(UInteractiveToolManager* Manager, UInteractiveTool* Tool) override;
+	virtual void OnToolEnded(UInteractiveToolManager* Manager, UInteractiveTool* Tool) override;
+	virtual FText GetActiveToolDisplayName() const override;
+	virtual FText GetActiveToolMessage() const override;
 
 private:
 	void PostNotification(const FText& Message);
@@ -46,6 +39,11 @@ private:
 
 	void UpdateActiveToolProperties(UInteractiveTool* Tool);
 
+	void RegisterPalettes();
+	FDelegateHandle ActivePaletteChangedHandle;
+
+	void MakeToolAcceptCancelWidget();
+
 	FText ActiveToolName;
 	FText ActiveToolMessage;
 	FStatusBarMessageHandle ActiveToolMessageHandle;
diff --git a/Engine/Plugins/Experimental/Animation/SkeletalMeshModelingTools/Source/SkeletalMeshModelingTools/Private/SkeletalMeshModelingToolsModule.cpp b/Engine/Plugins/Experimental/Animation/SkeletalMeshModelingTools/Source/SkeletalMeshModelingTools/Private/SkeletalMeshModelingToolsModule.cpp
--- a/Engine/Plugins/Experimental/Animation/SkeletalMeshModelingTools/Source/SkeletalMeshModelingTools/Private/SkeletalMeshModelingToolsModule.cpp
+++ b/Engine/Plugins/Experimental/Animation/SkeletalMeshModelingTools/Source/SkeletalMeshModelingTools/Private/SkeletalMeshModelingToolsModule.cpp
@@ -40,10 +40,14 @@ void FSkeletalMeshModelingToolsModule::StartupModule()
 		FSimpleMulticastDelegate::FDelegate::CreateRaw(this, &FSkeletalMeshModelingToolsModule::RegisterMenusAndToolbars));
 
 	ISkeletalMeshEditorModule& SkelMeshEditorModule = FModuleManager::Get().LoadModuleChecked<ISkeletalMeshEditorModule>("SkeletalMeshEditor");
+	// register toolbar extender with skeletal mesh editor
 	TArray<ISkeletalMeshEditorModule::FSkeletalMeshEditorToolbarExtender>& ToolbarExtenders = SkelMeshEditorModule.GetAllSkeletalMeshEditorToolbarExtenders();
-
 	ToolbarExtenders.Add(ISkeletalMeshEditorModule::FSkeletalMeshEditorToolbarExtender::CreateRaw(this, &FSkeletalMeshModelingToolsModule::ExtendSkelMeshEditorToolbar));
 	SkelMeshEditorExtenderHandle = ToolbarExtenders.Last().GetHandle();
+	// register post-init callback with skeletal mesh editor
+	TArray<ISkeletalMeshEditorModule::FOnSkeletalMeshEditorInitialized>& PostInitDelegates = SkelMeshEditorModule.GetPostEditorInitDelegates();
+	PostInitDelegates.Add(ISkeletalMeshEditorModule::FOnSkeletalMeshEditorInitialized::CreateRaw(this, &FSkeletalMeshModelingToolsModule::OnToggleEditingToolsMode));
+	SkelMeshEditorPostInitHandle = PostInitDelegates.Last().GetHandle();
 
 	FCoreDelegates::OnPostEngineInit.AddRaw(this, &FSkeletalMeshModelingToolsModule::OnPostEngineInit);
 
@@ -56,9 +60,13 @@ void FSkeletalMeshModelingToolsModule::ShutdownModule()
 	
 	if (ISkeletalMeshEditorModule* SkelMeshEditorModule = FModuleManager::GetModulePtr<ISkeletalMeshEditorModule>("SkeletalMeshEditor"))
 	{
+		// un-register toolbar extender delegates
 		TArray<ISkeletalMeshEditorModule::FSkeletalMeshEditorToolbarExtender>& Extenders = SkelMeshEditorModule->GetAllSkeletalMeshEditorToolbarExtenders();
-
 		Extenders.RemoveAll([=](const ISkeletalMeshEditorModule::FSkeletalMeshEditorToolbarExtender& InDelegate) { return InDelegate.GetHandle() == SkelMeshEditorExtenderHandle; });
+
+		// un-register post-init delegates
+		TArray<ISkeletalMeshEditorModule::FOnSkeletalMeshEditorInitialized>& PostInitDelegates = SkelMeshEditorModule->GetPostEditorInitDelegates();
+		PostInitDelegates.RemoveAll([=](const ISkeletalMeshEditorModule::FOnSkeletalMeshEditorInitialized& InDelegate) { return InDelegate.GetHandle() == SkelMeshEditorPostInitHandle; });
 	}
 
 	UToolMenus::UnregisterOwner(this);
@@ -121,12 +129,12 @@ TSharedRef<FExtender> FSkeletalMeshModelingToolsModule::ExtendSkelMeshEditorTool
 	// Add toolbar extender
 	TSharedPtr<FExtender> ToolbarExtender = MakeShareable(new FExtender);
 
-	TWeakPtr<ISkeletalMeshEditor> Ptr(InSkeletalMeshEditor);
+	TWeakPtr<ISkeletalMeshEditor> EditorPtr(InSkeletalMeshEditor);
 
 	InCommandList->MapAction(FSkeletalMeshModelingToolsCommands::Get().ToggleEditingToolsMode,
-	    FExecuteAction::CreateRaw(this, &FSkeletalMeshModelingToolsModule::OnToggleEditingToolsMode, Ptr),
+	    FExecuteAction::CreateRaw(this, &FSkeletalMeshModelingToolsModule::OnToggleEditingToolsMode, EditorPtr),
 	    FCanExecuteAction(),
-	    FIsActionChecked::CreateRaw(this, &FSkeletalMeshModelingToolsModule::IsEditingToolModeActive, Ptr));
+	    FIsActionChecked::CreateRaw(this, &FSkeletalMeshModelingToolsModule::IsEditingToolModeActive, EditorPtr));
 
 	return ToolbarExtender.ToSharedRef();
 }
diff --git a/Engine/Plugins/Experimental/Animation/SkeletalMeshModelingTools/Source/SkeletalMeshModelingTools/Private/SkeletalMeshModelingToolsModule.h b/Engine/Plugins/Experimental/Animation/SkeletalMeshModelingTools/Source/SkeletalMeshModelingTools/Private/SkeletalMeshModelingToolsModule.h
--- a/Engine/Plugins/Experimental/Animation/SkeletalMeshModelingTools/Source/SkeletalMeshModelingTools/Private/SkeletalMeshModelingToolsModule.h
+++ b/Engine/Plugins/Experimental/Animation/SkeletalMeshModelingTools/Source/SkeletalMeshModelingTools/Private/SkeletalMeshModelingToolsModule.h
@@ -42,6 +42,7 @@ private:
 	
 	// The handle for the extender delegate we added. Needed for clean module shutdown.
 	FDelegateHandle SkelMeshEditorExtenderHandle;
+	FDelegateHandle SkelMeshEditorPostInitHandle;
 
 	TArray<TWeakPtr<FApplicationMode>> RegisteredApplicationModes;
 };
diff --git a/Engine/Plugins/Experimental/Animation/SkeletalMeshModelingTools/Source/SkeletalMeshModelingTools/SkeletalMeshModelingTools.Build.cs b/Engine/Plugins/Experimental/Animation/SkeletalMeshModelingTools/Source/SkeletalMeshModelingTools/SkeletalMeshModelingTools.Build.cs
--- a/Engine/Plugins/Experimental/Animation/SkeletalMeshModelingTools/Source/SkeletalMeshModelingTools/SkeletalMeshModelingTools.Build.cs
+++ b/Engine/Plugins/Experimental/Animation/SkeletalMeshModelingTools/Source/SkeletalMeshModelingTools/SkeletalMeshModelingTools.Build.cs
@@ -34,7 +34,8 @@ public class SkeletalMeshModelingTools : ModuleRules
 				"PropertyEditor",
 				"AnimationCore",
 				"AnimationWidgets",
-				"ApplicationCore"
+				"ApplicationCore",
+				"WidgetRegistration",
 			}
 		);
 	}
diff --git a/Engine/Source/Editor/SkeletalMeshEditor/Private/SSkeletalMeshEditorToolbox.cpp b/Engine/Source/Editor/SkeletalMeshEditor/Private/SSkeletalMeshEditorToolbox.cpp
--- a/Engine/Source/Editor/SkeletalMeshEditor/Private/SSkeletalMeshEditorToolbox.cpp
+++ b/Engine/Source/Editor/SkeletalMeshEditor/Private/SSkeletalMeshEditorToolbox.cpp
@@ -4,230 +4,29 @@
 
 #include "ISkeletalMeshEditor.h"
 
-#include "EdMode.h"
 #include "Framework/MultiBox/MultiBoxBuilder.h"
-#include "Tools/UEdMode.h"
 #include "Widgets/Input/SCheckBox.h"
-#include "Widgets/Layout/SWidgetSwitcher.h"
 #include "Widgets/Layout/SUniformWrapPanel.h"
-#include "Widgets/Input/SSegmentedControl.h"
-
-
-SSkeletalMeshEditorToolbox::~SSkeletalMeshEditorToolbox()
-{
-}
-
 
 void SSkeletalMeshEditorToolbox::Construct(
 	const FArguments& InArgs, 
 	const TSharedRef<ISkeletalMeshEditor>& InOwningEditor
 	)
 {
-	SkeletalMeshEditor = InOwningEditor;
-
 	ChildSlot
 	[
-		SNew( SBorder )
+		SAssignNew(InlineContentHolder, SBorder)
 		.BorderImage( FAppStyle::GetBrush( "ToolPanel.GroupBorder" ) )
-		.Padding(0.0f)
-		[
-			SNew( SVerticalBox )
-			+ SVerticalBox::Slot()
-			.AutoHeight()
-			.HAlign( HAlign_Left )
-			[
-				SAssignNew( ModeToolBarContainer, SBorder )
-				.BorderImage( FAppStyle::GetBrush( "ToolPanel.GroupBorder" ) )
-				.Padding( FMargin(4, 0, 0, 0) )
-			]
-
-			+ SVerticalBox::Slot()
-			.FillHeight( 1.0f )
-			[
-				SNew( SVerticalBox )
-
-				+SVerticalBox::Slot()
-				.Padding(0.0, 8.0, 0.0, 0.0)
-				.AutoHeight()
-				[
-					SAssignNew(ModeToolHeader, SBorder)
-					.BorderImage( FAppStyle::GetBrush( "ToolPanel.GroupBorder" ) )
-				]
-
-				+ SVerticalBox::Slot()
-				.FillHeight(1)
-				[
-					SAssignNew(InlineContentHolder, SBorder)
-					.BorderImage( FAppStyle::GetBrush( "ToolPanel.GroupBorder" ) )
-					.Visibility( this, &SSkeletalMeshEditorToolbox::GetInlineContentHolderVisibility )
-				]
-			]
-		]
 	];
 }
 
-
 void SSkeletalMeshEditorToolbox::AttachToolkit(const TSharedRef<IToolkit>& InToolkit)
 {
-	UpdateInlineContent(InToolkit, InToolkit->GetInlineContent());
+	TSharedPtr<SWidget> Content = InToolkit->GetInlineContent();
+	InlineContentHolder->SetContent(Content.ToSharedRef());
 }
 
-
 void SSkeletalMeshEditorToolbox::DetachToolkit(const TSharedRef<IToolkit>& InToolkit)
 {
-	UpdateInlineContent(nullptr, SNullWidget::NullWidget);
-}
-
-
-void SSkeletalMeshEditorToolbox::SetOwningTab(TSharedRef<SDockTab>& InOwningTab)
-{
-	OwningTab = InOwningTab;
-}
-
-
-void SSkeletalMeshEditorToolbox::UpdateInlineContent(const TSharedPtr<IToolkit>& Toolkit, TSharedPtr<SWidget> InlineContent)
-{
-	static const FName SkeletalMeshEditorStatusBarName = "SkeletalMeshEditor.StatusBar";
-
-	// The display name that the owning tab should have as its label
-	FText TabName;
-
-	// The icon that should be displayed in the parent tab
-	const FSlateBrush* TabIcon = nullptr;
-
-	if (StatusBarMessageHandle.IsValid())
-	{
-		GEditor->GetEditorSubsystem<UStatusBarSubsystem>()->PopStatusBarMessage(SkeletalMeshEditorStatusBarName, StatusBarMessageHandle);
-		StatusBarMessageHandle.Reset();
-	}
-	TWeakPtr<FModeToolkit> ModeToolkit = StaticCastSharedPtr<FModeToolkit>(Toolkit);
-	if (ModeToolkit.IsValid())
-	{
-		TabName = ModeToolkit.Pin()->GetEditorModeDisplayName();
-		TabIcon = ModeToolkit.Pin()->GetEditorModeIcon().GetSmallIcon();
-		const TSharedRef<FModeToolkit> ModeToolkitPinned = ModeToolkit.Pin().ToSharedRef();
-
-		UpdatePalette(ModeToolkit);
-
-		// Show the name of the active tool in the statusbar.
-		// FIXME: We should also be showing Ctrl/Shift/Alt LMB/RMB shortcuts.
-		StatusBarMessageHandle = GEditor->GetEditorSubsystem<UStatusBarSubsystem>()->PushStatusBarMessage(
-			SkeletalMeshEditorStatusBarName,
-			TAttribute<FText>::Create(TAttribute<FText>::FGetter::CreateSP(ModeToolkitPinned, &FModeToolkit::GetActiveToolDisplayName)));
-	}
-	else
-	{
-		TabName = NSLOCTEXT("SkeletalMeshEditor", "ToolboxTab", "Toolbox");
-		TabIcon = FAppStyle::Get().GetBrush("LevelEditor.Tabs.Modes");
-	}
-
-	if (InlineContent.IsValid() && InlineContentHolder.IsValid())
-	{
-		InlineContentHolder->SetContent(InlineContent.ToSharedRef());
-	}
-
-	TSharedPtr<SDockTab> OwningTabPinned = OwningTab.Pin();
-	if (OwningTabPinned.IsValid())
-	{
-		OwningTabPinned->SetLabel(TabName);
-		OwningTabPinned->SetTabIcon(TabIcon);
-	}
-}
-
-
-void SSkeletalMeshEditorToolbox::UpdatePalette(const TWeakPtr<FModeToolkit>& InWeakToolkit) const
-{
-	if (!InWeakToolkit.IsValid())
-	{
-		return;
-	}
-	
-	TSharedRef<SSegmentedControl<FName>> PaletteTabBox = SNew(SSegmentedControl<FName>)
-		.UniformPadding(FMargin(8.f, 3.f))
-		.Value_Lambda([InWeakToolkit]()
-		{
-			return InWeakToolkit.IsValid() ? InWeakToolkit.Pin()->GetCurrentPalette() : NAME_None;
-		})
-		.OnValueChanged_Lambda([InWeakToolkit](const FName& Palette)
-		{
-			if (InWeakToolkit.IsValid())
-			{
-				InWeakToolkit.Pin()->SetCurrentPalette(Palette);
-			}
-		});
-
-	// Only show if there's more than one entry.
-	PaletteTabBox->SetVisibility(TAttribute<EVisibility>::Create(
-		TAttribute<EVisibility>::FGetter::CreateLambda([PaletteTabBox]() -> EVisibility { 
-			return PaletteTabBox->NumSlots() > 1 ? EVisibility::Visible : EVisibility::Collapsed;
-		})));
-
-	const TSharedPtr<FModeToolkit> ModeToolkit = InWeakToolkit.Pin();
-	
-	// Also build the toolkit here
-	TArray<FName> PaletteNames;
-	ModeToolkit->GetToolPaletteNames(PaletteNames);
-
-	const TSharedPtr<FUICommandList> CommandList = ModeToolkit->GetToolkitCommands();
-
-	TSharedRef< SWidgetSwitcher > PaletteSwitcher = SNew(SWidgetSwitcher)
-	.WidgetIndex_Lambda( [PaletteNames, InWeakToolkit] () -> int32 {
-		if (InWeakToolkit.IsValid())
-		{
-			int32 FoundIndex;
-			if (PaletteNames.Find(InWeakToolkit.Pin()->GetCurrentPalette(), FoundIndex))
-			{
-				return FoundIndex;	
-			}
-		}
-		return 0;
-	} );
-			
-	for(auto Palette : PaletteNames)
-	{
-		const FName ToolbarCustomizationName = ModeToolkit->GetEditorMode() ?
-			ModeToolkit->GetEditorMode()->GetModeInfo().ToolbarCustomizationName :
-			ModeToolkit->GetScriptableEditorMode()->GetModeInfo().ToolbarCustomizationName;
-		FUniformToolBarBuilder ModeToolbarBuilder(CommandList, FMultiBoxCustomization(ToolbarCustomizationName));
-		ModeToolbarBuilder.SetStyle(&FAppStyle::Get(), "PaletteToolBar");
-
-		ModeToolkit->BuildToolPalette(Palette, ModeToolbarBuilder);
-
-		TSharedRef<SWidget> PaletteWidget = ModeToolbarBuilder.MakeWidget();
-
-		const bool bRebuildChildren = false;
-		PaletteTabBox->AddSlot(Palette, false)
-		.Text(ModeToolkit->GetToolPaletteDisplayName(Palette));
-
-		PaletteSwitcher->AddSlot()
-		[
-			PaletteWidget
-		]; 
-	}
-
-	PaletteTabBox->RebuildChildren();
-
-	ModeToolHeader->SetContent(
-		SNew(SVerticalBox)
-
-		+SVerticalBox::Slot()
-		.Padding(8.f, 0.f, 0.f, 8.f)
-		.AutoHeight()
-		.HAlign(HAlign_Center)
-		[
-			PaletteTabBox
-		]
-
-		+SVerticalBox::Slot()
-		.AutoHeight()
-		[
-			PaletteSwitcher
-		]
-	);
-}
-
-
-EVisibility SSkeletalMeshEditorToolbox::GetInlineContentHolderVisibility() const
-{
-	return InlineContentHolder->GetContent() == SNullWidget::NullWidget ? EVisibility::Collapsed : EVisibility::Visible;
+	InlineContentHolder->SetContent(SNullWidget::NullWidget);
 }
diff --git a/Engine/Source/Editor/SkeletalMeshEditor/Private/SSkeletalMeshEditorToolbox.h b/Engine/Source/Editor/SkeletalMeshEditor/Private/SSkeletalMeshEditorToolbox.h
--- a/Engine/Source/Editor/SkeletalMeshEditor/Private/SSkeletalMeshEditorToolbox.h
+++ b/Engine/Source/Editor/SkeletalMeshEditor/Private/SSkeletalMeshEditorToolbox.h
@@ -17,14 +17,13 @@ class IToolkit;
 class SBorder;
 
 
-class SSkeletalMeshEditorToolbox :
-	public SCompoundWidget
+class SSkeletalMeshEditorToolbox : public SCompoundWidget
 {
 public:
 	SLATE_BEGIN_ARGS( SSkeletalMeshEditorToolbox ) {}
 	SLATE_END_ARGS()
 
-	~SSkeletalMeshEditorToolbox();
+	virtual ~SSkeletalMeshEditorToolbox() override = default;
 
 	void Construct(const FArguments& InArgs, const TSharedRef<ISkeletalMeshEditor>& InOwningEditor);
 
@@ -34,32 +33,7 @@ public:
 	// Called to detach a toolkit from this toolbox widget (e.g. when the mode changes).
 	void DetachToolkit(const TSharedRef<IToolkit>& InToolkit);
 
-	// Sets the owning tab so that we can update the tab icon based on what the toolkit provides.
-	void SetOwningTab(TSharedRef<SDockTab>& InOwningTab);
-
-private:	
-	void UpdateInlineContent(const TSharedPtr<IToolkit>& Toolkit, TSharedPtr<SWidget> InlineContent);
-
-	void UpdatePalette(const TWeakPtr<FModeToolkit>& InWeakToolkit) const;
-
-	// Whether to show the inline content or not.
-	EVisibility GetInlineContentHolderVisibility() const;
-
-	// The owning tab where the toolkit's interface is hosted. We need this to update the icons
-	// and title as the toolkit changes.
-	TWeakPtr<SDockTab> OwningTab;
-
-	TWeakPtr<ISkeletalMeshEditor> SkeletalMeshEditor;
-
-	// Inline content area for editor modes
+private:
+	// Inline content area for FToolkitBuilder contents
 	TSharedPtr<SBorder> InlineContentHolder;
-
-	// The container holding the mode toolbar
-	TSharedPtr<SBorder> ModeToolBarContainer;
-
-	// The active tool header area
-	TSharedPtr<SBorder> ModeToolHeader;
-
-	// The global status bar to show the name of the active tool.
-	FStatusBarMessageHandle StatusBarMessageHandle;
 };
\ No newline at end of file
diff --git a/Engine/Source/Editor/SkeletalMeshEditor/Private/SkeletalMeshEditor.cpp b/Engine/Source/Editor/SkeletalMeshEditor/Private/SkeletalMeshEditor.cpp
--- a/Engine/Source/Editor/SkeletalMeshEditor/Private/SkeletalMeshEditor.cpp
+++ b/Engine/Source/Editor/SkeletalMeshEditor/Private/SkeletalMeshEditor.cpp
@@ -285,6 +285,14 @@ void FSkeletalMeshEditor::InitSkeletalMeshEditor(const EToolkitMode::Type Mode,
 	// Make sure we get told when the editor mode changes so we can switch to the appropriate tab
 	// if there's a toolbox available.
 	GetEditorModeManager().OnEditorModeIDChanged().AddSP(this, &FSkeletalMeshEditor::OnEditorModeIdChanged);
+
+	// run attached post-init delegates
+	ISkeletalMeshEditorModule& SkeletalMeshEditorModule = FModuleManager::GetModuleChecked<ISkeletalMeshEditorModule>("SkeletalMeshEditor");
+	const TArray<ISkeletalMeshEditorModule::FOnSkeletalMeshEditorInitialized>& PostInitDelegates = SkeletalMeshEditorModule.GetPostEditorInitDelegates();
+	for (const auto& PostInitDelegate : PostInitDelegates)
+	{
+		PostInitDelegate.ExecuteIfBound(SharedThis<ISkeletalMeshEditor>(this));
+	}
 }
 
 FName FSkeletalMeshEditor::GetToolkitFName() const
@@ -392,18 +400,17 @@ TSharedRef<SDockTab> FSkeletalMeshEditor::SpawnToolboxTab(const FSpawnTabArgs& A
 	ToolboxWidget = SNew(SSkeletalMeshEditorToolbox, SharedThis<ISkeletalMeshEditor>(this));
 
 	TSharedRef<SDockTab> DockTab = SNew(SDockTab)
-	    .Label(LOCTEXT("ToolboxTab", "Toolbox"))
+		.Label(LOCTEXT("ToolboxTab", "Toolbox"))
 		[
 			SNew(SBox)
-	        .AddMetaData<FTagMetaData>(FTagMetaData(TEXT("ToolboxTab")))
-	        [
+			.AddMetaData<FTagMetaData>(FTagMetaData(TEXT("ToolboxTab")))
+			[
 				ToolboxWidget.ToSharedRef()
 			]
 		];
-
-	ToolboxWidget->SetOwningTab(DockTab);
-	ToolboxWidget->AttachToolkit(HostedToolkit.ToSharedRef());
 	
+	ToolboxWidget->AttachToolkit(HostedToolkit.ToSharedRef());
+
 	return DockTab;
 }
 
diff --git a/Engine/Source/Editor/SkeletalMeshEditor/Private/SkeletalMeshEditorMode.cpp b/Engine/Source/Editor/SkeletalMeshEditor/Private/SkeletalMeshEditorMode.cpp
--- a/Engine/Source/Editor/SkeletalMeshEditor/Private/SkeletalMeshEditorMode.cpp
+++ b/Engine/Source/Editor/SkeletalMeshEditor/Private/SkeletalMeshEditorMode.cpp
@@ -56,10 +56,9 @@ FSkeletalMeshEditorMode::FSkeletalMeshEditorMode(TSharedRef<FWorkflowCentricAppl
 				(
 					FTabManager::NewStack()
 					->SetSizeCoefficient(0.2f)
-					->SetHideTabWell(false)
-					->AddTab(SkeletalMeshEditorTabs::AssetDetailsTab, ETabState::OpenedTab)
-					->AddTab(SkeletalMeshEditorTabs::SkeletonTreeTab, ETabState::OpenedTab)
-					->SetForegroundTab(SkeletalMeshEditorTabs::AssetDetailsTab)
+					->SetHideTabWell(true)
+					->AddTab(SkeletalMeshEditorTabs::ToolboxDetailsTab, ETabState::OpenedTab)
+					->SetForegroundTab(SkeletalMeshEditorTabs::ToolboxDetailsTab)
 				)
 				->Split
 				(
@@ -91,20 +90,20 @@ FSkeletalMeshEditorMode::FSkeletalMeshEditorMode(TSharedRef<FWorkflowCentricAppl
 						FTabManager::NewStack()
 						->SetSizeCoefficient(0.5f)
 						->SetHideTabWell(false)
-						->AddTab(SkeletalMeshEditorTabs::DetailsTab, ETabState::OpenedTab)
-						->AddTab(SkeletalMeshEditorTabs::ToolboxDetailsTab, ETabState::ClosedTab)
-						->AddTab(SkeletalMeshEditorTabs::AdvancedPreviewTab, ETabState::OpenedTab)
-						->SetForegroundTab(SkeletalMeshEditorTabs::DetailsTab)
+						->AddTab(SkeletalMeshEditorTabs::SkeletonTreeTab, ETabState::OpenedTab)
+						->AddTab(SkeletalMeshEditorTabs::MorphTargetsTab, ETabState::OpenedTab)
+						->AddTab(SkeletalMeshEditorTabs::CurveMetadataTab, ETabState::OpenedTab)
+						->SetForegroundTab(SkeletalMeshEditorTabs::SkeletonTreeTab)
 					)
 					->Split
 					(
-
 						FTabManager::NewStack()
 						->SetSizeCoefficient(0.5f)
 						->SetHideTabWell(false)
-						->AddTab(SkeletalMeshEditorTabs::MorphTargetsTab, ETabState::OpenedTab)
-						->AddTab(SkeletalMeshEditorTabs::CurveMetadataTab, ETabState::OpenedTab)
-						->SetForegroundTab(SkeletalMeshEditorTabs::MorphTargetsTab)
+						->AddTab(SkeletalMeshEditorTabs::AssetDetailsTab, ETabState::OpenedTab)
+						->AddTab(SkeletalMeshEditorTabs::DetailsTab, ETabState::OpenedTab)
+						->AddTab(SkeletalMeshEditorTabs::AdvancedPreviewTab, ETabState::OpenedTab)
+						->SetForegroundTab(SkeletalMeshEditorTabs::AssetDetailsTab)
 					)
 				)
 			)
diff --git a/Engine/Source/Editor/SkeletalMeshEditor/Private/SkeletalMeshEditorModule.cpp b/Engine/Source/Editor/SkeletalMeshEditor/Private/SkeletalMeshEditorModule.cpp
--- a/Engine/Source/Editor/SkeletalMeshEditor/Private/SkeletalMeshEditorModule.cpp
+++ b/Engine/Source/Editor/SkeletalMeshEditor/Private/SkeletalMeshEditorModule.cpp
@@ -36,6 +36,7 @@ public:
 		return SkeletalMeshEditor;
 	}
 
+	virtual TArray<FOnSkeletalMeshEditorInitialized>& GetPostEditorInitDelegates() override { return SkeletalMeshEditorInitialized; }
 	virtual TArray<FSkeletalMeshEditorToolbarExtender>& GetAllSkeletalMeshEditorToolbarExtenders() override { return SkeletalMeshEditorToolbarExtenders; }
 
 	/** Gets the extensibility managers for outside entities to extend this editor's menus and toolbars */
@@ -46,6 +47,7 @@ private:
 	TSharedPtr<FExtensibilityManager> MenuExtensibilityManager;
 	TSharedPtr<FExtensibilityManager> ToolBarExtensibilityManager;
 
+	TArray<FOnSkeletalMeshEditorInitialized> SkeletalMeshEditorInitialized;
 	TArray<FSkeletalMeshEditorToolbarExtender> SkeletalMeshEditorToolbarExtenders;
 };
 
diff --git a/Engine/Source/Editor/SkeletalMeshEditor/Public/ISkeletalMeshEditorModule.h b/Engine/Source/Editor/SkeletalMeshEditor/Public/ISkeletalMeshEditorModule.h
--- a/Engine/Source/Editor/SkeletalMeshEditor/Public/ISkeletalMeshEditorModule.h
+++ b/Engine/Source/Editor/SkeletalMeshEditor/Public/ISkeletalMeshEditorModule.h
@@ -20,4 +20,8 @@ public:
 	/** Get all toolbar extenders */
 	DECLARE_DELEGATE_RetVal_TwoParams(TSharedRef<FExtender>, FSkeletalMeshEditorToolbarExtender, const TSharedRef<FUICommandList> /*InCommandList*/, TSharedRef<ISkeletalMeshEditor> /*InSkeletalMeshEditor*/);
 	virtual TArray<FSkeletalMeshEditorToolbarExtender>& GetAllSkeletalMeshEditorToolbarExtenders() = 0;
+
+	/** Add callback after editor has been initialized */
+	DECLARE_DELEGATE_OneParam(FOnSkeletalMeshEditorInitialized, TWeakPtr<ISkeletalMeshEditor>);
+	virtual TArray<FOnSkeletalMeshEditorInitialized>& GetPostEditorInitDelegates() = 0;
 };
