Skeleton Editing Tool: editable skeleton tree

diff --git a/Engine/Plugins/Experimental/Animation/SkeletalMeshModelingTools/Source/SkeletalMeshModelingTools/Private/SReferenceSkeletonTree.cpp b/Engine/Plugins/Experimental/Animation/SkeletalMeshModelingTools/Source/SkeletalMeshModelingTools/Private/SReferenceSkeletonTree.cpp
new file mode 100644
--- /dev/null
+++ b/Engine/Plugins/Experimental/Animation/SkeletalMeshModelingTools/Source/SkeletalMeshModelingTools/Private/SReferenceSkeletonTree.cpp
@@ -0,0 +1,940 @@
+// Copyright Epic Games, Inc. All Rights Reserved.
+
+#include "SReferenceSkeletonTree.h"
+
+#include "SkeletonModifier.h"
+#include "SkeletalMeshModelingToolsCommands.h"
+
+#include "SPositiveActionButton.h"
+
+#include "Misc/TextFilterExpressionEvaluator.h"
+#include "Widgets/Input/SSearchBox.h"
+#include "Widgets/Text/SInlineEditableTextBlock.h"
+
+#define LOCTEXT_NAMESPACE "SReferenceSkeletonTree"
+
+namespace ReferenceSkeletonTreeLocals
+{
+
+FSkeletonModifierChange::FSkeletonModifierChange(const USkeletonModifier* InModifier)
+	: FCommandChange()
+	, PreChangeSkeleton(InModifier->GetReferenceSkeleton())
+	, PreBoneTracker(InModifier->GetBoneIndexTracker())
+	, PostChangeSkeleton(InModifier->GetReferenceSkeleton())
+	, PostBoneTracker(InModifier->GetBoneIndexTracker())
+{}
+
+void FSkeletonModifierChange::StoreSkeleton(const USkeletonModifier* InModifier)
+{
+	PostChangeSkeleton = InModifier->GetReferenceSkeleton();
+	PostBoneTracker = InModifier->GetBoneIndexTracker();
+}
+
+void FSkeletonModifierChange::Apply(UObject* Object)
+{ // redo
+	USkeletonModifier* Modifier = CastChecked<USkeletonModifier>(Object);
+	Modifier->ExternalUpdate(PostChangeSkeleton, PostBoneTracker);
+}
+
+void FSkeletonModifierChange::Revert(UObject* Object)
+{ // undo
+	USkeletonModifier* Modifier = CastChecked<USkeletonModifier>(Object);
+	Modifier->ExternalUpdate(PreChangeSkeleton, PreBoneTracker);
+}
+
+}
+
+FBoneElement::FBoneElement(const FName& InBoneName, TWeakObjectPtr<USkeletonModifier> InModifier)
+	: BoneName(InBoneName)
+	, WeakModifier(InModifier)
+{}
+
+void FBoneElement::RequestRename() const
+{
+	OnRenameRequested.ExecuteIfBound();
+}
+
+void SBoneItem::Construct(const FArguments& InArgs)
+{
+	// cf. FSkeletonTreeBoneItem::GenerateWidgetForNameColumn
+	
+	WeakTreeElement = InArgs._TreeElement;
+	WeakModifier = InArgs._WeakModifier;
+
+	const TSharedPtr<FBoneElement> TreeElement = WeakTreeElement.Pin();
+
+	const FSlateBrush* Brush = FAppStyle::Get().GetBrush("SkeletonTree.Bone");
+
+	TSharedPtr<SHorizontalBox> RowBox;
+
+	ChildSlot
+	[
+		SAssignNew(RowBox, SHorizontalBox)
+		+ SHorizontalBox::Slot()
+		.AutoWidth()
+		.Padding(FMargin(0.0f, 2.0f))
+		.VAlign(VAlign_Center)
+		.HAlign(HAlign_Center)
+		[
+			SNew(SImage)
+			.Image(Brush)
+			.ColorAndOpacity_Static(&SBoneItem::GetTextColor, InArgs._IsSelected)
+		]
+	];
+
+	TSharedPtr< SInlineEditableTextBlock > InlineWidget;
+	
+	RowBox->AddSlot()
+    .AutoWidth()
+	.Padding(4, 0, 0, 0)
+    .VAlign(VAlign_Center)
+    [
+	    SAssignNew(InlineWidget, SInlineEditableTextBlock)
+		.Text(this, &SBoneItem::GetName)
+		// .Font(TextFont)
+		.ColorAndOpacity_Static(&SBoneItem::GetTextColor, InArgs._IsSelected)
+		.OnVerifyTextChanged(this, &SBoneItem::OnVerifyNameChanged)
+		.OnTextCommitted(InArgs._Delegates.OnBoneNameCommitted)
+		.MultiLine(false)
+    ];
+
+	TreeElement->OnRenameRequested.BindSP(InlineWidget.Get(), &SInlineEditableTextBlock::EnterEditingMode);
+}
+
+FSlateColor SBoneItem::GetTextColor(FIsSelected InIsSelected)
+{
+	const bool bIsSelected = InIsSelected.IsBound() ? InIsSelected.Execute() : false;
+	if (bIsSelected)
+	{
+		return FAppStyle::Get().GetSlateColor("Colors.ForegroundInverted");
+	}
+	
+	return FSlateColor::UseForeground();
+}
+
+bool SBoneItem::OnVerifyNameChanged(const FText& InText, FText& OutErrorMessage) const
+{
+	if (!WeakModifier.IsValid())
+	{
+		return false;
+	}
+
+	const TSharedPtr<FBoneElement> BoneElementPtr = WeakTreeElement.Pin();
+	const FReferenceSkeleton& ReferenceSkeleton = WeakModifier->GetReferenceSkeleton();
+
+	const FName NewName = FName(InText.ToString());
+	const int32 BoneIndex = ReferenceSkeleton.FindRawBoneIndex(NewName);
+	return BoneIndex == INDEX_NONE;
+}
+
+FText SBoneItem::GetName() const
+{
+	return FText::FromName(WeakTreeElement.Pin()->BoneName);
+}
+
+void SReferenceSkeletonRow::Construct(const FArguments& InArgs, const TSharedRef<STableViewBase>& InOwnerTable)
+{
+	WeakTreeElement = InArgs._TreeElement;
+	WeakModifier = InArgs._WeakModifier;
+	Delegates = InArgs._Delegates;
+
+	const FSuperRowType::FArguments Args = FSuperRowType::FArguments()
+		.OnDragDetected(InArgs._Delegates.OnDragDetected)
+		.OnCanAcceptDrop(InArgs._Delegates.OnCanAcceptDrop)
+		.OnAcceptDrop(InArgs._Delegates.OnAcceptDrop)
+		.Style( FAppStyle::Get(), "TableView.AlternatingRow" );
+	
+	SMultiColumnTableRow::Construct(Args , InOwnerTable);
+}
+
+TSharedRef<SWidget> SReferenceSkeletonRow::GenerateWidgetForColumn(const FName& ColumnName)
+{
+	return SNew(SHorizontalBox)
+	+ SHorizontalBox::Slot()
+	.AutoWidth()
+	[
+		SNew( SExpanderArrow, SharedThis(this) )
+			.ShouldDrawWires(true)
+	]
+
+	+ SHorizontalBox::Slot()
+	.AutoWidth()
+	[
+		SNew(SBoneItem)
+			.WeakModifier(WeakModifier)
+			.TreeElement(WeakTreeElement)
+			.Delegates(Delegates)
+			.IsSelected(this, &SReferenceSkeletonRow::IsSelected)
+	];
+}
+
+TSharedRef<FBoneItemDragDropOp> FBoneItemDragDropOp::New(const TWeakPtr<FBoneElement>& InElement)
+{
+	TSharedRef<FBoneItemDragDropOp> Operation = MakeShared<FBoneItemDragDropOp>();
+		Operation->Element = InElement;
+		Operation->Construct();
+	return Operation;
+}
+
+TSharedPtr<SWidget> FBoneItemDragDropOp::GetDefaultDecorator() const
+{
+	return SNew(SBorder)
+        .Visibility(EVisibility::Visible)
+        .BorderImage(FAppStyle::GetBrush("Menu.Background"))
+        [
+            SNew(STextBlock)
+            .Text(FText::FromName(Element.Pin()->BoneName))
+        ];
+}
+
+SReferenceSkeletonTree::SReferenceSkeletonTree()
+{}
+
+SReferenceSkeletonTree::~SReferenceSkeletonTree()
+{
+	// GEditor->UnregisterForUndo(this);
+}
+
+void SReferenceSkeletonTree::Construct(const FArguments& InArgs)
+{
+	Modifier = InArgs._Modifier;
+	TextFilter = MakeShareable(new FTextFilterExpressionEvaluator(ETextFilterExpressionEvaluatorMode::BasicString));
+	CommandList = MakeShared<FUICommandList>();
+	BindCommands();
+	
+	FRefSkeletonTreeDelegates Delegates;
+	Delegates.OnCanAcceptDrop = FOnRefSkeletonTreeCanAcceptDrop::CreateSP(this, &SReferenceSkeletonTree::OnCanAcceptDrop);
+	Delegates.OnAcceptDrop = FOnRefSkeletonTreeAcceptDrop::CreateSP(this, &SReferenceSkeletonTree::OnAcceptDrop);
+	Delegates.OnDragDetected = FOnDragDetected::CreateSP(this, &SReferenceSkeletonTree::OnDragDetected);
+	Delegates.OnBoneRenamed = FOnBoneRenamed::CreateSP(this, &SReferenceSkeletonTree::OnBoneRenamed);
+	Delegates.OnBoneNameCommitted = FOnTextCommitted::CreateSP(this, &SReferenceSkeletonTree::OnNewBoneNameCommitted);
+	
+	ChildSlot
+	[
+		SNew( SOverlay )
+		+SOverlay::Slot()
+		[
+			SNew(SVerticalBox)
+			+SVerticalBox::Slot()
+			.AutoHeight()
+			.Padding(FMargin(0.f, 2.f))
+			[
+				SNew(SHorizontalBox)
+				+SHorizontalBox::Slot()
+				.AutoWidth()
+				.VAlign(VAlign_Center)
+				.Padding(FMargin(6.f, 0.0))
+				[
+					SNew(SPositiveActionButton)
+					.OnGetMenuContent(this, &SReferenceSkeletonTree::CreateAddNewMenu)
+					.Icon(FAppStyle::Get().GetBrush("Icons.Plus"))
+				]
+
+				+SHorizontalBox::Slot()
+				.FillWidth(1.0f)
+				[
+					SNew(SSearchBox)
+					.SelectAllTextWhenFocused(true)
+					.OnTextChanged( this, &SReferenceSkeletonTree::OnFilterTextChanged )
+					.HintText(LOCTEXT( "SearchBoxHint", "Search Reference Skeleton Tree..."))
+				]
+			]
+			
+			+SVerticalBox::Slot()
+			.Padding(0.0f, 0.0f)
+			[
+				SNew(SBorder)
+				.Padding(2.0f)
+				.BorderImage(FAppStyle::GetBrush("SCSEditor.TreePanel"))
+				[
+					SAssignNew(TreeView, SRefSkeletonTreeView)
+					.TreeItemsSource(&RootElements)
+					.SelectionMode(ESelectionMode::Multi)
+					.OnGenerateRow_Lambda( [this, Delegates](
+						TSharedPtr<FBoneElement> InItem,
+						const TSharedRef<STableViewBase>& OwnerTable) -> TSharedRef<ITableRow>
+					{
+						return SNew(SReferenceSkeletonRow, OwnerTable)
+							.WeakModifier(Modifier)
+							.TreeElement(InItem)
+							.Delegates(Delegates);
+					})
+					.OnGetChildren(this, &SReferenceSkeletonTree::HandleGetChildrenForTree)
+					.OnSelectionChanged(this, &SReferenceSkeletonTree::OnSelectionChanged)
+					.OnContextMenuOpening(this, &SReferenceSkeletonTree::CreateContextMenu)
+					.OnMouseButtonDoubleClick(this, &SReferenceSkeletonTree::OnItemDoubleClicked)
+					.OnSetExpansionRecursive(this, &SReferenceSkeletonTree::OnSetExpansionRecursive)
+					.ItemHeight(24)
+					.HighlightParentNodesForSelection(true)
+					.HeaderRow
+					(
+						SNew(SHeaderRow)
+						+ SHeaderRow::Column("Name")
+						.ShouldGenerateWidget(true)
+						.DefaultLabel(LOCTEXT("BoneNameLabel", "Name"))
+						.FillWidth(0.5f)
+					)
+				]
+			]
+		]
+    ];
+
+	static constexpr bool IsInitialSetup = true;
+	RefreshTreeView(IsInitialSetup);
+
+	GEditor->RegisterForUndo(this);
+}
+
+void SReferenceSkeletonTree::PostUndo(bool bSuccess)
+{
+	RefreshTreeView(true);
+}
+
+void SReferenceSkeletonTree::PostRedo(bool bSuccess)
+{
+	RefreshTreeView(true);
+}
+
+void SReferenceSkeletonTree::AddItemToSelection(const TSharedPtr<FBoneElement>& InItem)
+{
+	TreeView->SetItemSelection(InItem, true, ESelectInfo::Direct);
+	TreeView->RequestScrollIntoView(InItem);
+}
+
+void SReferenceSkeletonTree::RemoveItemFromSelection(const TSharedPtr<FBoneElement>& InItem)
+{
+	TreeView->SetItemSelection(InItem, false, ESelectInfo::Direct);
+}
+
+void SReferenceSkeletonTree::ReplaceItemInSelection(const FText& InOldName, const FText& InNewName)
+{
+	const FName OldName(*InOldName.ToString());
+	const FName NewName(*InNewName.ToString());
+	
+	for (const TSharedPtr<FBoneElement>& Item : AllElements)
+	{
+		// remove old selection
+		if (Item->BoneName.IsEqual(OldName))
+		{
+			TreeView->SetItemSelection(Item, false, ESelectInfo::Direct);
+		}
+		// add new selection
+		if (Item->BoneName.IsEqual(NewName))
+		{
+			TreeView->SetItemSelection(Item, true, ESelectInfo::Direct);
+		}
+	}
+}
+
+TArray<TSharedPtr<FBoneElement>> SReferenceSkeletonTree::GetSelectedItems() const
+{
+	return TreeView->GetSelectedItems();
+}
+
+bool SReferenceSkeletonTree::HasSelectedItems() const
+{
+	return TreeView->GetNumItemsSelected() > 0;
+}
+
+void SReferenceSkeletonTree::BindCommands()
+{
+	const FSkeletalMeshModelingToolsCommands& Commands = FSkeletalMeshModelingToolsCommands::Get();
+ 
+	CommandList->MapAction(Commands.NewBone,
+        FExecuteAction::CreateSP(this, &SReferenceSkeletonTree::HandleNewBone),
+        FCanExecuteAction::CreateSP(this, &SReferenceSkeletonTree::CanAddNewBone));
+	
+	CommandList->MapAction(Commands.RemoveBone,
+        FExecuteAction::CreateSP(this, &SReferenceSkeletonTree::HandleDeleteBone),
+        FCanExecuteAction::CreateSP(this, &SReferenceSkeletonTree::CanDeleteBone));
+ 
+	CommandList->MapAction(Commands.RenameBone,
+	FExecuteAction::CreateSP(this, &SReferenceSkeletonTree::HandleRenameBone),
+		FCanExecuteAction::CreateSP(this, &SReferenceSkeletonTree::CanRenameBone));
+	
+	CommandList->MapAction(Commands.UnParentBone,
+		FExecuteAction::CreateSP(this, &SReferenceSkeletonTree::HandleUnParentBone),
+		FCanExecuteAction::CreateSP(this, &SReferenceSkeletonTree::CanUnParentBone));
+}
+
+void SReferenceSkeletonTree::HandleNewBone()
+{
+	if (!Modifier.IsValid())
+	{
+		return;
+	}
+
+	BeginChange();
+	
+	TArray<FName> BoneNames;
+	GetSelectedBoneNames(BoneNames);
+	
+	const FName ParentName = BoneNames.IsEmpty() ? NAME_None : BoneNames[0];
+
+	static const FName DefaultName("joint");
+	const FName BoneName = Modifier->GetUniqueName(DefaultName);
+	const bool bAdded = Modifier->AddBone(BoneName, ParentName, FTransform::Identity);
+	if (bAdded)
+	{
+		RefreshTreeView(true);
+		if (Notifier.IsValid())
+		{
+			Notifier->Notify( {BoneName}, ESkeletalMeshNotifyType::BonesAdded);
+		}
+		EndChange();
+	}
+	else
+	{
+		CancelChange();
+	}
+}
+
+bool SReferenceSkeletonTree::CanAddNewBone() const
+{
+	return true;
+}
+
+void SReferenceSkeletonTree::HandleDeleteBone()
+{
+	if (!Modifier.IsValid())
+	{
+		return;
+	}
+
+	BeginChange();
+	
+	TArray<FName> BoneNames;
+	GetSelectedBoneNames(BoneNames);
+	
+	const bool bRemoved = Modifier->RemoveBones(BoneNames, true);
+	if (bRemoved)
+	{
+		RefreshTreeView(true);
+		if (Notifier.IsValid())
+		{
+			Notifier->Notify(BoneNames, ESkeletalMeshNotifyType::BonesRemoved);
+		}
+		EndChange();
+	}
+	else
+	{
+		CancelChange();
+	}
+}
+
+bool SReferenceSkeletonTree::CanDeleteBone() const
+{
+	return HasSelectedItems();
+}
+
+void SReferenceSkeletonTree::HandleUnParentBone()
+{
+	if (!Modifier.IsValid())
+	{
+		return;
+	}
+
+	BeginChange();
+	
+	TArray<FName> BoneNames;
+	GetSelectedBoneNames(BoneNames);
+
+	static const TArray<FName> Dummy;
+	const bool bUnParented = Modifier->ParentBones(BoneNames, Dummy);
+	if (bUnParented)
+	{
+		RefreshTreeView(true);
+		SelectItemFromNames(BoneNames);
+		if (Notifier.IsValid())
+		{
+			Notifier->Notify(BoneNames, ESkeletalMeshNotifyType::HierarchyChanged);
+		}
+		EndChange();
+	}
+	else
+	{
+		CancelChange();
+	}
+}
+
+bool SReferenceSkeletonTree::CanUnParentBone() const
+{
+	return HasSelectedItems();
+}
+
+void SReferenceSkeletonTree::GetSelectedBoneNames(TArray<FName>& OutSelectedBoneNames) const
+{
+	OutSelectedBoneNames.Reset();
+	const TArray<TSharedPtr<FBoneElement>> BoneElements = TreeView->GetSelectedItems();
+	OutSelectedBoneNames.Reserve(BoneElements.Num());
+
+	Algo::Transform(BoneElements, OutSelectedBoneNames, [](const TSharedPtr<FBoneElement>& BoneElement)
+	{
+		return BoneElement->BoneName;
+	});
+}
+
+void SReferenceSkeletonTree::SelectItemFromNames(const TArray<FName>& InBoneNames)
+{
+	for (const TSharedPtr<FBoneElement>& Item : AllElements)
+	{
+		const bool bSelect = InBoneNames.Contains(Item->BoneName);
+		TreeView->SetItemSelection(Item, bSelect, ESelectInfo::Direct);
+	}
+}
+
+ISkeletalMeshNotifier& SReferenceSkeletonTree::GetNotifier()
+{
+	if (!Notifier)
+	{
+		Notifier.Reset(new FReferenceSkeletonWidgetNotifier(SharedThis(this)));
+	}
+	return *Notifier;
+}
+
+void SReferenceSkeletonTree::HandleRenameBone() const
+{
+	TArray<TSharedPtr<FBoneElement>> BoneElements = TreeView->GetSelectedItems();
+	if (BoneElements.IsEmpty())
+	{
+		return;
+	}
+	BoneElements[0]->RequestRename();
+}
+
+bool SReferenceSkeletonTree::CanRenameBone() const
+{
+	return HasSelectedItems();
+}
+
+void SReferenceSkeletonTree::OnBoneRenamed(const FName InOldName, const FName InNewName) const
+{
+	if (!Modifier.IsValid())
+	{
+		return;
+	}
+
+	if (Notifier.IsValid())
+	{
+		Notifier->Notify( {InNewName}, ESkeletalMeshNotifyType::BonesRenamed);
+	}
+}
+
+void SReferenceSkeletonTree::OnNewBoneNameCommitted(const FText& InText, ETextCommit::Type InCommitType)
+{
+	TArray<FName> BoneNames;
+	GetSelectedBoneNames(BoneNames);
+	
+	const FName NewName = FName(InText.ToString());
+	
+	if (BoneNames.IsEmpty() || NewName == NAME_None)
+	{
+		return;
+	}
+	
+	if (BoneNames.Num() == 1 && BoneNames[0] == NewName)
+	{
+		return;
+	}
+
+	TArray<int32> BoneIndices;
+	const FReferenceSkeleton& ReferenceSkeleton = Modifier->GetReferenceSkeleton();
+	Algo::Transform(BoneNames, BoneIndices, [&](const FName& BoneName)
+	{
+		return ReferenceSkeleton.FindRawBoneIndex(BoneName);
+	});
+
+	TArray<FName> NewNames;
+	NewNames.Init(NewName, BoneNames.Num());
+	
+	BeginChange();
+
+	const bool bBoneRenamed = Modifier->RenameBones(BoneNames, NewNames);
+	if (bBoneRenamed)
+	{
+		const TArray<FMeshBoneInfo>& BoneInfos = ReferenceSkeleton.GetRawRefBoneInfo();
+		for (int32 Index = 0; Index < BoneIndices.Num(); Index++)
+		{
+			if (BoneIndices[Index] != INDEX_NONE)
+			{
+				const FName OldName = BoneNames[Index];
+				BoneNames[Index] = BoneInfos[BoneIndices[Index]].Name;
+				const int32 ItemIndex = AllElements.IndexOfByPredicate([OldName](const TSharedPtr<FBoneElement>& Item)
+				{
+					return Item->BoneName == OldName;
+				});
+
+				if (ItemIndex != INDEX_NONE)
+				{
+					AllElements[ItemIndex]->BoneName = BoneNames[Index];
+				}
+			}
+		}
+		
+		if (Notifier.IsValid())
+		{
+			Notifier->Notify(BoneNames, ESkeletalMeshNotifyType::BonesRenamed);
+		}
+		
+		EndChange();
+	}
+	else
+	{
+		CancelChange();	
+	}
+}
+
+void SReferenceSkeletonTree::OnFilterTextChanged(const FText& SearchText)
+{
+	TextFilter->SetFilterText(SearchText);
+	RefreshTreeView();
+}
+
+void SReferenceSkeletonTree::RefreshTreeView(bool IsInitialSetup /*=false*/)
+{
+	if (!Modifier.IsValid())
+	{
+		return;
+	}
+	
+	const FReferenceSkeleton& RefSkeleton = Modifier->GetReferenceSkeleton();
+	const TArray<FMeshBoneInfo>& BoneInfos = RefSkeleton.GetRawRefBoneInfo();
+	
+	// reset all tree items
+	RootElements.Reset();
+	AllElements.Reset();
+
+	// validate we have a skeleton to load
+	if (BoneInfos.IsEmpty())
+	{
+		TreeView->RequestTreeRefresh();
+		return;
+	}
+	
+	// record bone element indices
+	TMap<FName, int32> BoneTreeElementIndices;
+
+	auto FilterString = [this](const FString& StringToTest) ->bool
+	{
+		return TextFilter->TestTextFilter(FBasicStringFilterExpressionContext(StringToTest));
+	};
+
+	// create all bone elements
+	for (int32 BoneIndex=0; BoneIndex < BoneInfos.Num(); ++BoneIndex)
+	{
+		const FName BoneName = BoneInfos[BoneIndex].Name;
+		TSharedPtr<FBoneElement> BoneElement = MakeShared<FBoneElement>(BoneName, Modifier);
+		const int32 BoneElementIndex = AllElements.Add(BoneElement);
+		BoneTreeElementIndices.Add(BoneName, BoneElementIndex);
+
+		// store pointer to parent (if there is one)
+		const int32 ParentIndex = BoneInfos[BoneIndex].ParentIndex;
+		if (ParentIndex != INDEX_NONE)
+		{
+			// get parent tree element
+			const FName ParentBoneName = BoneInfos[ParentIndex].Name;
+			const TSharedPtr<FBoneElement> ParentBoneTreeElement = AllElements[BoneTreeElementIndices[ParentBoneName]];
+			BoneElement->UnFilteredParent = ParentBoneTreeElement;
+		}
+
+		// apply text filter to bones
+		if (!(TextFilter->GetFilterText().IsEmpty() || FilterString(BoneName.ToString())))
+		{
+			BoneElement->bIsHidden = true;
+		}
+	}
+
+	// resolve parent/children pointers on all tree elements, taking into consideration the filter options
+	// (elements are parented to their nearest non-hidden/filtered parent element)
+	for (TSharedPtr<FBoneElement>& Element : AllElements)
+	{
+		if (Element->bIsHidden)
+		{
+			continue;
+		}
+		
+		// find first parent that is not filtered
+		TSharedPtr<FBoneElement> ParentElement = Element->UnFilteredParent;
+		while (true)
+		{
+			if (!ParentElement.IsValid())
+			{
+				break;
+			}
+
+			if (!ParentElement->bIsHidden)
+			{
+				break;
+			}
+
+			ParentElement = ParentElement->UnFilteredParent;
+		}
+
+		if (ParentElement.IsValid())
+		{
+			// store pointer to child on parent
+			ParentElement->Children.Add(Element);
+			// store pointer to parent on child
+			Element->Parent = ParentElement;
+		}
+		else
+		{
+			// has no parent, store a root element
+			RootElements.Add(Element);
+		}
+	}
+
+	// expand all elements upon the initial construction of the tree
+	for (const TSharedPtr<FBoneElement> RootElement : RootElements)
+	{
+		SetExpansionRecursive(RootElement, false, true);
+	}
+	
+	TreeView->RequestTreeRefresh();
+}
+
+void SReferenceSkeletonTree::HandleGetChildrenForTree(
+	TSharedPtr<FBoneElement> InItem,
+	TArray<TSharedPtr<FBoneElement>>& OutChildren)
+{
+	OutChildren = InItem.Get()->Children;
+}
+
+void SReferenceSkeletonTree::OnSelectionChanged(TSharedPtr<FBoneElement> Selection, ESelectInfo::Type SelectInfo)
+{
+	TArray<FName> BoneNames;
+	GetSelectedBoneNames(BoneNames);
+	
+	if (Notifier.IsValid())
+	{
+		Notifier->Notify(BoneNames, ESkeletalMeshNotifyType::BonesSelected);
+	}
+}
+
+TSharedRef<SWidget> SReferenceSkeletonTree::CreateAddNewMenu()
+{
+	const FSkeletalMeshModelingToolsCommands& Commands = FSkeletalMeshModelingToolsCommands::Get();
+
+	static constexpr bool bCloseAfterSelection = true;
+	FMenuBuilder MenuBuilder(bCloseAfterSelection, CommandList);
+	
+	MenuBuilder.BeginSection("NewBone", LOCTEXT("AddNewBoneOperations", "Bones"));
+	MenuBuilder.AddMenuEntry(Commands.NewBone);
+	MenuBuilder.EndSection();
+		
+	return MenuBuilder.MakeWidget();
+}
+
+TSharedPtr<SWidget> SReferenceSkeletonTree::CreateContextMenu()
+{
+	static constexpr bool bCloseAfterSelection = true;
+	FMenuBuilder MenuBuilder(bCloseAfterSelection, CommandList);
+
+	const FSkeletalMeshModelingToolsCommands& Commands = FSkeletalMeshModelingToolsCommands::Get();
+	
+	const TArray<TSharedPtr<FBoneElement>> SelectedItems = GetSelectedItems();
+	if (SelectedItems.IsEmpty())
+	{
+		MenuBuilder.BeginSection("NewBone", LOCTEXT("AddNewBoneOperations", "Bones"));
+	
+		MenuBuilder.AddMenuEntry(Commands.NewBone);
+
+		MenuBuilder.EndSection();
+	}
+	else
+	{
+		MenuBuilder.BeginSection("EditBones", LOCTEXT("EditBonesOperations", "Bones"));
+
+		MenuBuilder.AddMenuEntry(Commands.NewBone);
+		MenuBuilder.AddMenuEntry(Commands.RemoveBone);
+		MenuBuilder.AddMenuEntry(Commands.UnParentBone);
+		MenuBuilder.AddMenuEntry(Commands.RenameBone);
+
+		MenuBuilder.EndSection();
+	}
+	
+	return MenuBuilder.MakeWidget();
+}
+
+void SReferenceSkeletonTree::OnItemDoubleClicked(TSharedPtr<FBoneElement> InItem)
+{
+	InItem->RequestRename();
+}
+
+void SReferenceSkeletonTree::OnSetExpansionRecursive(TSharedPtr<FBoneElement> InItem, bool bShouldBeExpanded)
+{
+	SetExpansionRecursive(InItem, false, bShouldBeExpanded);
+}
+
+void SReferenceSkeletonTree::SetExpansionRecursive(
+	TSharedPtr<FBoneElement> InElement,
+	bool bTowardsParent,
+    bool bShouldBeExpanded)
+{
+	TreeView->SetItemExpansion(InElement, bShouldBeExpanded);
+    
+    if (bTowardsParent)
+    {
+    	if (InElement->Parent.Get())
+    	{
+    		SetExpansionRecursive(InElement->Parent, bTowardsParent, bShouldBeExpanded);
+    	}
+    }
+    else
+    {
+    	for (int32 ChildIndex = 0; ChildIndex < InElement->Children.Num(); ++ChildIndex)
+    	{
+    		SetExpansionRecursive(InElement->Children[ChildIndex], bTowardsParent, bShouldBeExpanded);
+    	}
+    }
+}
+
+FReply SReferenceSkeletonTree::OnKeyDown(const FGeometry& MyGeometry, const FKeyEvent& InKeyEvent)
+{
+	if (!Modifier.IsValid())
+	{
+		return FReply::Handled();
+	}
+	
+	if (CommandList.IsValid() && CommandList->ProcessCommandBindings(InKeyEvent))
+	{
+		return FReply::Handled();
+	}
+	
+	return FReply::Unhandled();
+}
+
+FReply SReferenceSkeletonTree::OnDragDetected(const FGeometry& MyGeometry, const FPointerEvent& MouseEvent)
+{
+	const TArray<TSharedPtr<FBoneElement>> BoneElements = GetSelectedItems();
+	if (BoneElements.Num() != 1)
+	{
+		return FReply::Unhandled();
+	}
+	
+	if (MouseEvent.IsMouseButtonDown(EKeys::LeftMouseButton))
+	{
+		const TSharedRef<FBoneItemDragDropOp> DragDropOp = FBoneItemDragDropOp::New(BoneElements[0]);
+		return FReply::Handled().BeginDragDrop(DragDropOp);
+	}
+	
+	return FReply::Unhandled();
+}
+
+TOptional<EItemDropZone> SReferenceSkeletonTree::OnCanAcceptDrop(const FDragDropEvent& DragDropEvent, EItemDropZone DropZone, TSharedPtr<FBoneElement> TargetItem)
+{
+	TOptional<EItemDropZone> ReturnedDropZone;
+
+	const TSharedPtr<FBoneItemDragDropOp> DragDropOp = DragDropEvent.GetOperationAs<FBoneItemDragDropOp>();
+	if (DragDropOp.IsValid())
+	{
+		ReturnedDropZone = EItemDropZone::OntoItem;
+	}
+
+	return ReturnedDropZone;
+}
+
+FReply SReferenceSkeletonTree::OnAcceptDrop(const FDragDropEvent& DragDropEvent, EItemDropZone DropZone, TSharedPtr<FBoneElement> TargetItem)
+{
+	const TSharedPtr<FBoneItemDragDropOp> DragDropOp = DragDropEvent.GetOperationAs<FBoneItemDragDropOp>();
+	if (!DragDropOp.IsValid())
+	{
+		return FReply::Unhandled();
+	}
+
+	if (!Modifier.IsValid())
+	{
+		return FReply::Handled();
+	}
+	
+	const FName BoneName = DragDropOp->Element.Pin()->BoneName;
+	const FName ParentName = TargetItem->BoneName;
+
+	BeginChange();
+	
+	const bool bParented = Modifier->ParentBone(BoneName, ParentName);
+	if (bParented)
+	{
+		RefreshTreeView(true);
+		SelectItemFromNames({BoneName});
+		if (Notifier.IsValid())
+		{
+			Notifier->Notify( {BoneName}, ESkeletalMeshNotifyType::HierarchyChanged);
+		}
+		EndChange();
+	}
+	else
+	{
+		CancelChange();
+	}
+	
+	return FReply::Handled();
+}
+
+FReferenceSkeletonWidgetNotifier::FReferenceSkeletonWidgetNotifier(TSharedRef<SReferenceSkeletonTree> InWidget)
+	: ISkeletalMeshNotifier()
+	, Tree(InWidget)
+{}
+
+void FReferenceSkeletonWidgetNotifier::HandleNotification(const TArray<FName>& InBoneNames, const ESkeletalMeshNotifyType InNotifyType)
+{
+	if (!Tree.IsValid())
+	{
+		return;
+	}
+
+	TSharedPtr<SReferenceSkeletonTree> TreePtr = Tree.Pin();
+	switch (InNotifyType)
+	{
+		case ESkeletalMeshNotifyType::BonesAdded:
+			TreePtr->RefreshTreeView();
+			break;
+		case ESkeletalMeshNotifyType::BonesRemoved:
+			TreePtr->RefreshTreeView();
+			break;
+		case ESkeletalMeshNotifyType::BonesMoved:
+			break;
+		case ESkeletalMeshNotifyType::BonesSelected:
+			TreePtr->SelectItemFromNames(InBoneNames);
+			break;
+		case ESkeletalMeshNotifyType::BonesRenamed:
+			TreePtr->RefreshTreeView();
+			TreePtr->SelectItemFromNames(InBoneNames);
+			break;
+		case ESkeletalMeshNotifyType::HierarchyChanged:
+			{
+				TArray<FName> BoneNames;
+				TreePtr->GetSelectedBoneNames(BoneNames);
+				TreePtr->RefreshTreeView();
+				TreePtr->SelectItemFromNames(BoneNames);
+			}
+			break;
+	}
+}
+
+void SReferenceSkeletonTree::BeginChange()
+{
+	ensure( ActiveChange == nullptr );
+	ActiveChange = MakeUnique<ReferenceSkeletonTreeLocals::FSkeletonModifierChange>(Modifier.Get()); 
+}
+
+void SReferenceSkeletonTree::EndChange()
+{
+	if (!ActiveChange.IsValid())
+	{
+		return;
+	}
+
+	ActiveChange->StoreSkeleton(Modifier.Get());
+	
+	GEditor->BeginTransaction(LOCTEXT("ModifyRefSkeleton", "Modify Reference Skeleton"));
+	GUndo->StoreUndo(Modifier.Get(), MoveTemp(ActiveChange));
+	GEditor->EndTransaction();
+}
+
+void SReferenceSkeletonTree::CancelChange()
+{
+	ActiveChange.Reset();
+}
+
+#undef LOCTEXT_NAMESPACE
\ No newline at end of file
diff --git a/Engine/Plugins/Experimental/Animation/SkeletalMeshModelingTools/Source/SkeletalMeshModelingTools/Private/SReferenceSkeletonTree.h b/Engine/Plugins/Experimental/Animation/SkeletalMeshModelingTools/Source/SkeletalMeshModelingTools/Private/SReferenceSkeletonTree.h
new file mode 100644
--- /dev/null
+++ b/Engine/Plugins/Experimental/Animation/SkeletalMeshModelingTools/Source/SkeletalMeshModelingTools/Private/SReferenceSkeletonTree.h
@@ -0,0 +1,249 @@
+// Copyright Epic Games, Inc. All Rights Reserved.
+
+#pragma once
+
+#include "EditorUndoClient.h"
+#include "SkeletalMeshNotifier.h"
+#include "DragAndDrop/DecoratedDragDropOp.h"
+#include "Widgets/SCompoundWidget.h"
+#include "Widgets/Views/SListView.h"
+#include "Widgets/Views/STableRow.h"
+#include "Widgets/Views/STreeView.h"
+
+class FTextFilterExpressionEvaluator;
+class FUICommandList;
+class USkeletonModifier;
+
+namespace ReferenceSkeletonTreeLocals
+{
+
+class FSkeletonModifierChange : public FCommandChange
+{
+public:
+	FSkeletonModifierChange(const USkeletonModifier* InModifier);
+
+	virtual FString ToString() const override
+	{
+		return FString(TEXT("Edit Skeleton"));
+	}
+
+	virtual void Apply(UObject* Object) override;
+	virtual void Revert(UObject* Object) override;
+
+	void StoreSkeleton(const USkeletonModifier* InModifier);
+
+private:
+	FReferenceSkeleton PreChangeSkeleton;
+	TArray<int32> PreBoneTracker;
+	FReferenceSkeleton PostChangeSkeleton;
+	TArray<int32> PostBoneTracker;
+};
+
+}
+
+class FBoneElement : public TSharedFromThis<FBoneElement>
+{
+public:
+	FBoneElement(const FName& InBoneName, TWeakObjectPtr<USkeletonModifier> InModifier);
+
+	FName BoneName;
+
+	bool bIsHidden = false;
+	
+	TSharedPtr<FBoneElement> UnFilteredParent;
+	TSharedPtr<FBoneElement> Parent;
+	TArray<TSharedPtr<FBoneElement>> Children;
+
+	void RequestRename() const;
+	
+	DECLARE_DELEGATE(FOnRenameRequested);
+	FOnRenameRequested OnRenameRequested;
+
+	DECLARE_DELEGATE_TwoParams(FOnRenamed, const FName, const FName);
+	FOnRenamed OnRenamed;
+
+private:
+	TWeakObjectPtr<USkeletonModifier> WeakModifier;
+};
+
+using FOnRefSkeletonTreeCanAcceptDrop = STableRow<TSharedPtr<FBoneElement>>::FOnCanAcceptDrop;
+using FOnRefSkeletonTreeAcceptDrop = STableRow<TSharedPtr<FBoneElement>>::FOnAcceptDrop;
+using FOnBoneRenamed = FBoneElement::FOnRenamed;
+
+struct FRefSkeletonTreeDelegates
+{
+	FOnDragDetected					OnDragDetected;
+	FOnRefSkeletonTreeCanAcceptDrop	OnCanAcceptDrop;
+	FOnRefSkeletonTreeAcceptDrop	OnAcceptDrop;
+	FOnBoneRenamed					OnBoneRenamed;
+	FOnTextCommitted				OnBoneNameCommitted;
+};
+
+class SBoneItem : public SCompoundWidget
+{
+	SLATE_BEGIN_ARGS(SBoneItem) {}
+		SLATE_ARGUMENT(TWeakPtr<FBoneElement>, TreeElement)
+		SLATE_ARGUMENT(TWeakObjectPtr<USkeletonModifier>, WeakModifier)
+		SLATE_ARGUMENT(FRefSkeletonTreeDelegates, Delegates)
+		SLATE_EVENT( FIsSelected, IsSelected )
+	SLATE_END_ARGS()
+
+	// Slate construction function
+	void Construct(const FArguments& InArgs);
+	
+private:
+
+	bool OnVerifyNameChanged(const FText& InText, FText& OutErrorMessage) const;
+
+	static FSlateColor GetTextColor(FIsSelected InIsSelected);
+	
+	FText GetName() const;
+	
+	TWeakPtr<FBoneElement> WeakTreeElement;
+	TWeakObjectPtr<USkeletonModifier> WeakModifier;
+};
+
+class FBoneItemDragDropOp : public FDecoratedDragDropOp
+{
+public:
+	DRAG_DROP_OPERATOR_TYPE(FBoneItemDragDropOp, FDecoratedDragDropOp)
+
+	static TSharedRef<FBoneItemDragDropOp> New(const TWeakPtr<FBoneElement>& InElement);
+
+	// FDecoratedDragDropOp interface
+	virtual TSharedPtr<SWidget> GetDefaultDecorator() const override;
+	// End FDecoratedDragDropOp interface
+	
+	TWeakPtr<FBoneElement> Element;
+};
+
+class SReferenceSkeletonRow : public SMultiColumnTableRow<TSharedPtr<FBoneElement>>
+{
+public:
+	SLATE_BEGIN_ARGS(SReferenceSkeletonRow) {}
+		SLATE_ARGUMENT(TWeakObjectPtr<USkeletonModifier>, WeakModifier)
+		SLATE_ARGUMENT(TSharedPtr<FBoneElement>, TreeElement)
+		SLATE_ARGUMENT(FRefSkeletonTreeDelegates, Delegates)
+	SLATE_END_ARGS()
+
+	// Slate construction function
+	void Construct(const FArguments& InArgs, const TSharedRef<STableViewBase>& InOwnerTable);
+
+	// SMultiColumnTableRow interface
+	virtual TSharedRef<SWidget> GenerateWidgetForColumn(const FName& InColumnName) override;
+	// End SMultiColumnTableRow interface
+
+private:
+	TWeakPtr<FBoneElement> WeakTreeElement;
+	TWeakObjectPtr<USkeletonModifier> WeakModifier;
+	FRefSkeletonTreeDelegates Delegates;
+};
+
+typedef STreeView< TSharedPtr<FBoneElement> > SRefSkeletonTreeView;
+
+class SReferenceSkeletonTree;
+
+class FReferenceSkeletonWidgetNotifier: public ISkeletalMeshNotifier
+{
+public:
+	FReferenceSkeletonWidgetNotifier(TSharedRef<SReferenceSkeletonTree> InTree);
+	virtual void HandleNotification(const TArray<FName>& BoneNames, const ESkeletalMeshNotifyType InNotifyType) override;
+	
+private:
+	TWeakPtr<SReferenceSkeletonTree> Tree;
+};
+
+class SReferenceSkeletonTree : public SCompoundWidget, public FEditorUndoClient
+{
+public:
+	SLATE_BEGIN_ARGS(SReferenceSkeletonTree) {}
+		SLATE_ARGUMENT(TWeakObjectPtr<USkeletonModifier>, Modifier)
+		SLATE_ARGUMENT(FRefSkeletonTreeDelegates, Delegates)
+	SLATE_END_ARGS()
+
+	SReferenceSkeletonTree();
+    virtual ~SReferenceSkeletonTree() override;
+
+	// Slate construction function
+	void Construct(const FArguments& InArgs);
+
+	virtual void PostUndo(bool bSuccess) override;
+	virtual void PostRedo(bool bSuccess) override;
+	
+	// Selection
+	void GetSelectedBoneNames(TArray<FName>& OutSelectedBoneNames) const;
+	TArray<TSharedPtr<FBoneElement>> GetSelectedItems() const;
+	bool HasSelectedItems() const;
+	void SelectItemFromNames(const TArray<FName>& InBoneNames);
+
+	ISkeletalMeshNotifier& GetNotifier();
+
+private:
+	// SWidget interface
+	virtual FReply OnKeyDown(const FGeometry& MyGeometry, const FKeyEvent& InKeyEvent) override;
+	// END SWidget interface
+	
+	void BindCommands();
+
+	void AddItemToSelection(const TSharedPtr<FBoneElement>& InItem);
+	void RemoveItemFromSelection(const TSharedPtr<FBoneElement>& InItem);
+	void ReplaceItemInSelection(const FText& OldName, const FText& NewName);
+
+	// Create, Rename, Delete, Parent bones
+	void HandleNewBone();
+	bool CanAddNewBone() const;
+	void HandleRenameBone() const;
+	bool CanRenameBone() const;
+	void HandleDeleteBone();
+	bool CanDeleteBone() const;
+	void HandleUnParentBone();
+	bool CanUnParentBone() const;
+
+	// Filter
+	void OnFilterTextChanged(const FText& SearchText);
+	
+	// Callbacks
+	void RefreshTreeView(bool IsInitialSetup=false);
+	void HandleGetChildrenForTree(TSharedPtr<FBoneElement> InItem, TArray<TSharedPtr<FBoneElement>>& OutChildren);
+	void OnSelectionChanged(TSharedPtr<FBoneElement> Selection, ESelectInfo::Type SelectInfo);
+	
+	TSharedRef< SWidget > CreateAddNewMenu();
+	TSharedPtr< SWidget > CreateContextMenu();
+
+	void OnItemDoubleClicked(TSharedPtr<FBoneElement> InItem);
+	void OnSetExpansionRecursive(TSharedPtr<FBoneElement> InItem, bool bShouldBeExpanded);
+	void SetExpansionRecursive(TSharedPtr<FBoneElement> InElement, bool bTowardsParent, bool bShouldBeExpanded);
+
+	// Delegates
+	virtual FReply OnDragDetected(const FGeometry& MyGeometry, const FPointerEvent& MouseEvent) override;
+	TOptional<EItemDropZone> OnCanAcceptDrop(const FDragDropEvent& DragDropEvent, EItemDropZone DropZone, TSharedPtr<FBoneElement> TargetItem);
+	FReply OnAcceptDrop(const FDragDropEvent& DragDropEvent, EItemDropZone DropZone, TSharedPtr<FBoneElement> TargetItem);
+	void OnBoneRenamed(const FName InOldName, const FName InNewName) const;
+	
+	void OnNewBoneNameCommitted(const FText& InText, ETextCommit::Type InCommitType);
+
+	// The modifier that actually does modifications on the reference skeleton
+	TWeakObjectPtr<USkeletonModifier> Modifier;
+	
+	TUniquePtr<FReferenceSkeletonWidgetNotifier> Notifier;
+
+	// undo
+	void BeginChange();
+	void EndChange();
+	void CancelChange();
+	TUniquePtr<ReferenceSkeletonTreeLocals::FSkeletonModifierChange> ActiveChange;
+	
+	// Commands
+	TSharedPtr<FUICommandList> CommandList;
+	
+	// Tree view
+	TSharedPtr<SRefSkeletonTreeView> TreeView;
+	TArray<TSharedPtr<FBoneElement>> RootElements;
+	TArray<TSharedPtr<FBoneElement>> AllElements;
+
+	// filter
+	TSharedPtr<FTextFilterExpressionEvaluator> TextFilter;
+
+	friend class FReferenceSkeletonWidgetNotifier;
+};
+
diff --git a/Engine/Plugins/Experimental/Animation/SkeletalMeshModelingTools/Source/SkeletalMeshModelingTools/Private/SkeletalMeshEditorUtils.cpp b/Engine/Plugins/Experimental/Animation/SkeletalMeshModelingTools/Source/SkeletalMeshModelingTools/Private/SkeletalMeshEditorUtils.cpp
--- a/Engine/Plugins/Experimental/Animation/SkeletalMeshModelingTools/Source/SkeletalMeshModelingTools/Private/SkeletalMeshEditorUtils.cpp
+++ b/Engine/Plugins/Experimental/Animation/SkeletalMeshModelingTools/Source/SkeletalMeshModelingTools/Private/SkeletalMeshEditorUtils.cpp
@@ -6,6 +6,10 @@
 #include "InteractiveToolsContext.h"
 #include "IPersonaToolkit.h"
 #include "ISkeletalMeshEditor.h"
+#include "ISkeletonEditorModule.h"
+#include "SkeletonModifier.h"
+#include "SReferenceSkeletonTree.h"
+#include "Preferences/PersonaOptions.h"
 
 #include UE_INLINE_GENERATED_CPP_BY_NAME(SkeletalMeshEditorUtils)
 
@@ -66,7 +70,8 @@ void USkeletalMeshEditorContextObject::Unregister(UInteractiveToolManager* InToo
 	
 	InToolManager->GetContextObjectStore()->RemoveContextObject(this);
 
-	Bindings.Reset();
+	EditorBindings.Reset();
+	TreeBindings.Reset();
 	
 	bRegistered = false;
 }
@@ -74,7 +79,8 @@ void USkeletalMeshEditorContextObject::Unregister(UInteractiveToolManager* InToo
 void USkeletalMeshEditorContextObject::Init(const TWeakPtr<ISkeletalMeshEditor>& InEditor)
 {
 	Editor = InEditor;
-	Bindings.Reset();
+	EditorBindings.Reset();
+	TreeBindings.Reset();
 }
 
 void USkeletalMeshEditorContextObject::HideSkeleton()
@@ -115,74 +121,186 @@ void USkeletalMeshEditorContextObject::ShowSkeleton()
 
 void USkeletalMeshEditorContextObject::BindTo(ISkeletalMeshEditingInterface* InEditingInterface)
 {
-	if (!InEditingInterface || !Editor.IsValid())
+	if (!InEditingInterface)
 	{
 		return;
 	}
 	
-	if (Bindings.Contains(InEditingInterface))
+	BindEditor(InEditingInterface);
+	BindRefSkeletonTree(InEditingInterface);
+}
+
+void USkeletalMeshEditorContextObject::UnbindFrom(ISkeletalMeshEditingInterface* InEditingInterface)
+{
+	if (!InEditingInterface)
+	{
+		return;
+	}
+
+	UnbindEditor(InEditingInterface);
+	UnbindRefSkeletonTree(InEditingInterface);
+}
+
+TPair<FDelegateHandle, FDelegateHandle> USkeletalMeshEditorContextObject::BindInterfaceTo(
+	ISkeletalMeshEditingInterface* InInterface,
+	ISkeletalMeshNotifier& InOtherNotifier)
+{
+	// connect external interface to tool (ie skeletal mesh editor -> tool)
+	FDelegateHandle ToToolNotifierHandle = InOtherNotifier.Delegate().AddLambda(
+	[InInterface](const TArray<FName>& BoneNames, const ESkeletalMeshNotifyType InNotifyType)
+	{
+		InInterface->GetNotifier().HandleNotification(BoneNames, InNotifyType);
+	});
+
+	// connect tool to external interface (ie tool -> skeletal mesh editor)
+	FDelegateHandle FromToolNotifierHandle = InInterface->GetNotifier().Delegate().AddLambda(
+		[&InOtherNotifier](const TArray<FName>& BoneNames, const ESkeletalMeshNotifyType InNotifyType)
+	{
+		InOtherNotifier.HandleNotification(BoneNames, InNotifyType);
+	});
+
+	return {ToToolNotifierHandle, FromToolNotifierHandle};
+}
+
+void USkeletalMeshEditorContextObject::UnbindInterfaceFrom(
+	ISkeletalMeshEditingInterface* InInterface,
+	ISkeletalMeshNotifier& InOtherNotifier,
+	const FBindData& InOutBindData)
+{
+	if (InOutBindData.ToToolNotifierHandle.IsValid())
+	{
+		InOtherNotifier.Delegate().Remove(InOutBindData.ToToolNotifierHandle);
+	}
+
+	if (InOutBindData.FromToolNotifierHandle.IsValid())
+	{
+		InInterface->GetNotifier().Delegate().Remove(InOutBindData.FromToolNotifierHandle);
+	}
+}
+
+void USkeletalMeshEditorContextObject::BindEditor(ISkeletalMeshEditingInterface* InEditingInterface)
+{
+	if (!InEditingInterface || !Editor.IsValid())
 	{
 		return;
 	}
 	
+	if (EditorBindings.Contains(InEditingInterface))
+	{
+		return;
+	}
+
 	TSharedPtr<ISkeletalMeshEditorBinding> Binding = Editor.Pin()->GetBinding();
-	if (Binding.IsValid())
+	if (!Binding.IsValid())
 	{
-		InEditingInterface->BindTo(Binding);
-	
-		FBindData BindData;
-	
-		// connect external interface to tool (ie skeletal mesh editor -> tool)
-		BindData.ToToolNotifierHandle = Binding->GetNotifier().Delegate().AddLambda(
-			[InEditingInterface](const TArray<FName>& BoneNames, const ESkeletalMeshNotifyType InNotifyType)
+		return;
+	}
+
+	InEditingInterface->BindTo(Binding);
+
+	FBindData BindData;
+	Tie(BindData.ToToolNotifierHandle, BindData.FromToolNotifierHandle) = BindInterfaceTo(InEditingInterface, Binding->GetNotifier());
+	EditorBindings.Emplace(InEditingInterface, BindData);
+
+	InEditingInterface->GetNotifier().HandleNotification(Binding->GetSelectedBones(), ESkeletalMeshNotifyType::BonesSelected);
+}
+
+void USkeletalMeshEditorContextObject::UnbindEditor(ISkeletalMeshEditingInterface* InEditingInterface)
+{
+	if (const FBindData* BindData = EditorBindings.Find(InEditingInterface))
+	{
+		if (BindData->FromToolNotifierHandle.IsValid())
 		{
-			InEditingInterface->GetNotifier().HandleNotification(BoneNames, InNotifyType);
-		});
-	
-		// connect too to external interface (ie tool -> skeletal mesh editor)
-		BindData.FromToolNotifierHandle = InEditingInterface->GetNotifier().Delegate().AddLambda(
-			[Binding](const TArray<FName>& BoneNames, const ESkeletalMeshNotifyType InNotifyType)
+			InEditingInterface->Unbind();
+		}
+
+		if (Editor.IsValid())
 		{
-			Binding->GetNotifier().HandleNotification(BoneNames, InNotifyType);
-		});
-		Bindings.Emplace(InEditingInterface, BindData);
+			const TSharedPtr<ISkeletalMeshEditorBinding> Binding = Editor.Pin()->GetBinding();
+			UnbindInterfaceFrom(InEditingInterface, Binding->GetNotifier(), *BindData);
+		}
 	
-		InEditingInterface->GetNotifier().HandleNotification(Binding->GetSelectedBones(), ESkeletalMeshNotifyType::BonesSelected);
+		EditorBindings.Remove(InEditingInterface);
 	}
 }
 
-void USkeletalMeshEditorContextObject::UnbindFrom(ISkeletalMeshEditingInterface* InEditingInterface)
+void USkeletalMeshEditorContextObject::BindRefSkeletonTree(ISkeletalMeshEditingInterface* InEditingInterface)
 {
-	if (!InEditingInterface)
+	if (!InEditingInterface || !Editor.IsValid())
+	{
+		return;
+	}
+	
+	const TWeakObjectPtr<USkeletonModifier> Modifier = InEditingInterface->GetModifier();
+	if (!Modifier.IsValid())
 	{
 		return;
 	}
 	
-	if (FBindData* BindData = Bindings.Find(InEditingInterface))
+	const TSharedPtr<FTabManager> TabManager = Editor.Pin()->GetAssociatedTabManager();
+	TSharedPtr<SDockTab> SkeletonTab = TabManager->FindExistingLiveTab(GetSkeletonTreeTabId());
+	if (SkeletonTab.IsValid())
+	{
+		DefaultSkeletonWidget = SkeletonTab->GetContent();
+	}
+	else
+	{
+		SkeletonTab = TabManager->TryInvokeTab(GetSkeletonTreeTabId());
+	}
+
+	check(SkeletonTab.IsValid());
+
+	SAssignNew(RefSkeletonWidget, SBorder)
+		.Padding(4.f)
+		.BorderImage(FAppStyle::GetBrush("ToolPanel.GroupBorder"))
+	[
+		SAssignNew(RefSkeletonTree, SReferenceSkeletonTree)
+			.Modifier(Modifier)
+	];
+
+	SkeletonTab->SetContent(RefSkeletonWidget.ToSharedRef());
+
+	FBindData BindData;
+	Tie(BindData.ToToolNotifierHandle, BindData.FromToolNotifierHandle) = BindInterfaceTo(InEditingInterface, RefSkeletonTree->GetNotifier());
+	TreeBindings.Emplace(InEditingInterface, BindData);
+}
+
+void USkeletalMeshEditorContextObject::UnbindRefSkeletonTree(ISkeletalMeshEditingInterface* InEditingInterface)
+{
+	if (const FBindData* BindData = TreeBindings.Find(InEditingInterface))
 	{
-		if (BindData->ToToolNotifierHandle.IsValid())
+		if (RefSkeletonTree.IsValid())
 		{
-			TSharedPtr<ISkeletalMeshEditorBinding> Binding = GetBinding();
-			if (Binding.IsValid())
-			{
-				Binding->GetNotifier().Delegate().Remove(BindData->ToToolNotifierHandle);
-			}
-			BindData->ToToolNotifierHandle.Reset();
+			UnbindInterfaceFrom(InEditingInterface, RefSkeletonTree->GetNotifier(), *BindData);
 		}
-	
-		if (BindData->FromToolNotifierHandle.IsValid())
+
+		if (Editor.IsValid())
 		{
-			InEditingInterface->GetNotifier().Delegate().Remove(BindData->FromToolNotifierHandle);
-			InEditingInterface->Unbind();
-			
-			BindData->FromToolNotifierHandle.Reset();
+			const TSharedPtr<FTabManager> TabManager = Editor.Pin()->GetAssociatedTabManager();
+			const TSharedPtr<SDockTab> SkeletonTab = TabManager->FindExistingLiveTab(GetSkeletonTreeTabId());
+			if (SkeletonTab.IsValid())
+			{
+				if (SkeletonTab->GetContent() == RefSkeletonWidget)
+				{
+					SkeletonTab->SetContent(DefaultSkeletonWidget.IsValid() ? DefaultSkeletonWidget.ToSharedRef() : SNullWidget::NullWidget);
+				}
+			}
 		}
-	
-		Bindings.Remove(InEditingInterface);
+		
+		RefSkeletonTree.Reset();
+		RefSkeletonWidget.Reset();
+		
+		TreeBindings.Remove(InEditingInterface);
 	}
 }
 
+const FName& USkeletalMeshEditorContextObject::GetSkeletonTreeTabId()
+{
+	static const FName SkeletonTreeId(TEXT("SkeletonTreeView"));
+	return SkeletonTreeId;
+}
+
 TSharedPtr<ISkeletalMeshEditorBinding> USkeletalMeshEditorContextObject::GetBinding() const
 {
 	return Editor.IsValid() ? Editor.Pin()->GetBinding() : nullptr;
-}
\ No newline at end of file
+}
diff --git a/Engine/Plugins/Experimental/Animation/SkeletalMeshModelingTools/Source/SkeletalMeshModelingTools/Private/SkeletalMeshEditorUtils.h b/Engine/Plugins/Experimental/Animation/SkeletalMeshModelingTools/Source/SkeletalMeshModelingTools/Private/SkeletalMeshEditorUtils.h
--- a/Engine/Plugins/Experimental/Animation/SkeletalMeshModelingTools/Source/SkeletalMeshModelingTools/Private/SkeletalMeshEditorUtils.h
+++ b/Engine/Plugins/Experimental/Animation/SkeletalMeshModelingTools/Source/SkeletalMeshModelingTools/Private/SkeletalMeshEditorUtils.h
@@ -14,6 +14,7 @@ class UInteractiveToolsContext;
 class UInteractiveToolManager;
 class ISkeletalMeshEditor;
 class USkeletalMeshEditorContextObject;
+class SReferenceSkeletonTree;
 
 namespace UE
 {
@@ -62,7 +63,29 @@ private:
 		FDelegateHandle ToToolNotifierHandle;
 		FDelegateHandle FromToolNotifierHandle;		
 	};
-	TMap< ISkeletalMeshEditingInterface*, FBindData > Bindings;
+	TMap< ISkeletalMeshEditingInterface*, FBindData > EditorBindings;
+	TMap< ISkeletalMeshEditingInterface*, FBindData > TreeBindings;
+
+	static TPair<FDelegateHandle, FDelegateHandle> BindInterfaceTo(
+		ISkeletalMeshEditingInterface* InInterface,
+		ISkeletalMeshNotifier& InOtherNotifier);
+	
+	static void UnbindInterfaceFrom(
+		ISkeletalMeshEditingInterface* InInterface,
+		ISkeletalMeshNotifier& InOtherNotifier,
+		const FBindData& InOutBindData);
+
+	void BindEditor(ISkeletalMeshEditingInterface* InEditingInterface);
+	void UnbindEditor(ISkeletalMeshEditingInterface* InEditingInterface);
+
+	void BindRefSkeletonTree(ISkeletalMeshEditingInterface* InEditingInterface);
+	void UnbindRefSkeletonTree(ISkeletalMeshEditingInterface* InEditingInterface);
+
+	static const FName& GetSkeletonTreeTabId();
+
+	TSharedPtr<SWidget> RefSkeletonWidget;
+	TSharedPtr<SReferenceSkeletonTree> RefSkeletonTree;
+	TSharedPtr<SWidget> DefaultSkeletonWidget;
 	
 	bool bRegistered = false;
 };
diff --git a/Engine/Plugins/Experimental/Animation/SkeletalMeshModelingTools/Source/SkeletalMeshModelingTools/Private/SkeletalMeshModelingToolsCommands.cpp b/Engine/Plugins/Experimental/Animation/SkeletalMeshModelingTools/Source/SkeletalMeshModelingTools/Private/SkeletalMeshModelingToolsCommands.cpp
--- a/Engine/Plugins/Experimental/Animation/SkeletalMeshModelingTools/Source/SkeletalMeshModelingTools/Private/SkeletalMeshModelingToolsCommands.cpp
+++ b/Engine/Plugins/Experimental/Animation/SkeletalMeshModelingTools/Source/SkeletalMeshModelingTools/Private/SkeletalMeshModelingToolsCommands.cpp
@@ -10,6 +10,11 @@
 void FSkeletalMeshModelingToolsCommands::RegisterCommands()
 {
 	UI_COMMAND(ToggleEditingToolsMode, "Enable Editing Tools", "Toggles editing tools on or off.", EUserInterfaceActionType::ToggleButton, FInputChord());
+	
+	UI_COMMAND(NewBone, "New Bone", "Add new Bone.", EUserInterfaceActionType::Button, FInputChord(EKeys::N, EModifierKey::Control));
+	UI_COMMAND(RemoveBone, "Remove Bone", "Remove selected bone(s).", EUserInterfaceActionType::Button, FInputChord(EKeys::Delete), FInputChord(EKeys::BackSpace));
+	UI_COMMAND(UnParentBone, "Unparent Bone", "Unparent selected bone(s).", EUserInterfaceActionType::Button, FInputChord(EKeys::P, EModifierKey::Shift));
+	UI_COMMAND(RenameBone, "Rename Bone", "Rename the selected bone.", EUserInterfaceActionType::Button, FInputChord(EKeys::F2));
 }
 
 const FSkeletalMeshModelingToolsCommands& FSkeletalMeshModelingToolsCommands::Get()
diff --git a/Engine/Plugins/Experimental/Animation/SkeletalMeshModelingTools/Source/SkeletalMeshModelingTools/Private/SkeletalMeshModelingToolsCommands.h b/Engine/Plugins/Experimental/Animation/SkeletalMeshModelingTools/Source/SkeletalMeshModelingTools/Private/SkeletalMeshModelingToolsCommands.h
--- a/Engine/Plugins/Experimental/Animation/SkeletalMeshModelingTools/Source/SkeletalMeshModelingTools/Private/SkeletalMeshModelingToolsCommands.h
+++ b/Engine/Plugins/Experimental/Animation/SkeletalMeshModelingTools/Source/SkeletalMeshModelingTools/Private/SkeletalMeshModelingToolsCommands.h
@@ -25,6 +25,12 @@ public:
 
 	// Editing tools commands
 	TSharedPtr<FUICommandInfo> ToggleEditingToolsMode;
+
+	// Skeleton Editing commands
+	TSharedPtr<FUICommandInfo> NewBone;
+	TSharedPtr<FUICommandInfo> RemoveBone;
+	TSharedPtr<FUICommandInfo> UnParentBone;
+	TSharedPtr<FUICommandInfo> RenameBone;
 };
 
 class FSkeletalMeshModelingToolsActionCommands : public TInteractiveToolCommands<FSkeletalMeshModelingToolsActionCommands>
diff --git a/Engine/Plugins/Experimental/Animation/SkeletalMeshModelingTools/Source/SkeletalMeshModelingTools/SkeletalMeshModelingTools.Build.cs b/Engine/Plugins/Experimental/Animation/SkeletalMeshModelingTools/Source/SkeletalMeshModelingTools/SkeletalMeshModelingTools.Build.cs
--- a/Engine/Plugins/Experimental/Animation/SkeletalMeshModelingTools/Source/SkeletalMeshModelingTools/SkeletalMeshModelingTools.Build.cs
+++ b/Engine/Plugins/Experimental/Animation/SkeletalMeshModelingTools/Source/SkeletalMeshModelingTools/SkeletalMeshModelingTools.Build.cs
@@ -36,6 +36,7 @@ public class SkeletalMeshModelingTools : ModuleRules
 				"AnimationWidgets",
 				"ApplicationCore",
 				"WidgetRegistration",
+				"SkeletalMeshModifiers"
 			}
 		);
 	}
diff --git a/Engine/Plugins/Experimental/MeshModelingToolsetExp/Source/MeshModelingToolsEditorOnlyExp/Private/SkeletalMesh/SkeletalMeshEditionInterface.cpp b/Engine/Plugins/Experimental/MeshModelingToolsetExp/Source/MeshModelingToolsEditorOnlyExp/Private/SkeletalMesh/SkeletalMeshEditionInterface.cpp
--- a/Engine/Plugins/Experimental/MeshModelingToolsetExp/Source/MeshModelingToolsEditorOnlyExp/Private/SkeletalMesh/SkeletalMeshEditionInterface.cpp
+++ b/Engine/Plugins/Experimental/MeshModelingToolsetExp/Source/MeshModelingToolsEditorOnlyExp/Private/SkeletalMesh/SkeletalMeshEditionInterface.cpp
@@ -29,6 +29,11 @@ void ISkeletalMeshEditingInterface::Unbind()
 	Binding.Reset();
 }
 
+TWeakObjectPtr<USkeletonModifier> ISkeletalMeshEditingInterface::GetModifier() const
+{
+	return nullptr;
+}
+
 TOptional<FName> ISkeletalMeshEditingInterface::GetBoneName(HHitProxy* InHitProxy) const
 {
 	if (const HBoneHitProxy* BoneProxy = HitProxyCast<HBoneHitProxy>(InHitProxy))
diff --git a/Engine/Plugins/Experimental/MeshModelingToolsetExp/Source/MeshModelingToolsEditorOnlyExp/Private/SkeletalMesh/SkeletonEditingTool.cpp b/Engine/Plugins/Experimental/MeshModelingToolsetExp/Source/MeshModelingToolsEditorOnlyExp/Private/SkeletalMesh/SkeletonEditingTool.cpp
--- a/Engine/Plugins/Experimental/MeshModelingToolsetExp/Source/MeshModelingToolsEditorOnlyExp/Private/SkeletalMesh/SkeletonEditingTool.cpp
+++ b/Engine/Plugins/Experimental/MeshModelingToolsetExp/Source/MeshModelingToolsEditorOnlyExp/Private/SkeletalMesh/SkeletonEditingTool.cpp
@@ -499,6 +499,11 @@ void USkeletonEditingTool::UnParentBones()
 	if (bBonesUnParented)
 	{
 		GetToolManager()->DisplayMessage(LOCTEXT("Unparent", "Selected bones have been unparented."), EToolMessageLevel::UserNotification);
+
+		if (NeedsNotification())
+		{
+			GetNotifier().Notify({}, ESkeletalMeshNotifyType::HierarchyChanged);
+		}
 		
 		EndChange();
 		return;
@@ -518,6 +523,11 @@ void USkeletonEditingTool::ParentBones(const FName& InParentName)
 	const bool bBonesParented = Modifier->ParentBones(GetSelectedBones(), {InParentName});
 	if (bBonesParented)
 	{
+		if (NeedsNotification())
+		{
+			GetNotifier().Notify({}, ESkeletalMeshNotifyType::HierarchyChanged);
+		}
+		
 		EndChange();
 	}
 	else
@@ -595,11 +605,13 @@ void USkeletonEditingTool::RenameBones()
 				Selection[Index] = BoneInfos[BoneIndices[Index]].Name;
 			}
 		}
-		Properties->Name = GetCurrentBone();
+
+		const FName NewName = GetCurrentBone();
+		Properties->Name = NewName;
 		
 		if (NeedsNotification())
 		{
-			GetNotifier().Notify({CurrentBone}, ESkeletalMeshNotifyType::BonesRenamed);
+			GetNotifier().Notify( {NewName}, ESkeletalMeshNotifyType::BonesRenamed);
 		}
 		
 		EndChange();
@@ -919,6 +931,11 @@ FInputRayHit USkeletonEditingTool::CanBeginClickDragSequence(const FInputDeviceR
 	return FInputRayHit();
 }
 
+TWeakObjectPtr<USkeletonModifier> USkeletonEditingTool::GetModifier() const
+{
+	return Modifier;
+}
+
 void USkeletonEditingTool::HandleSkeletalMeshModified(const TArray<FName>& InBoneNames, const ESkeletalMeshNotifyType InNotifyType)
 {
 	if (GetNotifier().Notifying())
@@ -953,6 +970,8 @@ void USkeletonEditingTool::HandleSkeletalMeshModified(const TArray<FName>& InBon
 		case ESkeletalMeshNotifyType::BonesRenamed:
 			Selection = BoneNames;
 			break;
+		case ESkeletalMeshNotifyType::HierarchyChanged:
+			break;
 		default:
 			break;
 	}
diff --git a/Engine/Plugins/Experimental/MeshModelingToolsetExp/Source/MeshModelingToolsEditorOnlyExp/Private/SkeletalMesh/SkinWeightsBindingTool.cpp b/Engine/Plugins/Experimental/MeshModelingToolsetExp/Source/MeshModelingToolsEditorOnlyExp/Private/SkeletalMesh/SkinWeightsBindingTool.cpp
--- a/Engine/Plugins/Experimental/MeshModelingToolsetExp/Source/MeshModelingToolsEditorOnlyExp/Private/SkeletalMesh/SkinWeightsBindingTool.cpp
+++ b/Engine/Plugins/Experimental/MeshModelingToolsetExp/Source/MeshModelingToolsEditorOnlyExp/Private/SkeletalMesh/SkinWeightsBindingTool.cpp
@@ -376,6 +376,10 @@ void USkinWeightsBindingTool::HandleSkeletalMeshModified(const TArray<FName>& In
 			Properties->CurrentBone = InBoneNames.IsEmpty() ? NAME_None : InBoneNames[0];
 		}
 		break;
+	case ESkeletalMeshNotifyType::BonesRenamed:
+		break;
+	case ESkeletalMeshNotifyType::HierarchyChanged:
+        break;
 	default:
 		checkNoEntry();
 	}
diff --git a/Engine/Plugins/Experimental/MeshModelingToolsetExp/Source/MeshModelingToolsEditorOnlyExp/Private/SkeletalMesh/SkinWeightsPaintTool.cpp b/Engine/Plugins/Experimental/MeshModelingToolsetExp/Source/MeshModelingToolsEditorOnlyExp/Private/SkeletalMesh/SkinWeightsPaintTool.cpp
--- a/Engine/Plugins/Experimental/MeshModelingToolsetExp/Source/MeshModelingToolsEditorOnlyExp/Private/SkeletalMesh/SkinWeightsPaintTool.cpp
+++ b/Engine/Plugins/Experimental/MeshModelingToolsetExp/Source/MeshModelingToolsEditorOnlyExp/Private/SkeletalMesh/SkinWeightsPaintTool.cpp
@@ -1798,6 +1798,10 @@ void USkinWeightsPaintTool::HandleSkeletalMeshModified(const TArray<FName>& InBo
 			}
 		}
 		break;
+	case ESkeletalMeshNotifyType::BonesRenamed:
+		break;
+	case ESkeletalMeshNotifyType::HierarchyChanged:
+		break;
 	default:
 		checkNoEntry();
 	}
diff --git a/Engine/Plugins/Experimental/MeshModelingToolsetExp/Source/MeshModelingToolsEditorOnlyExp/Public/SkeletalMesh/SkeletalMeshEditionInterface.h b/Engine/Plugins/Experimental/MeshModelingToolsetExp/Source/MeshModelingToolsEditorOnlyExp/Public/SkeletalMesh/SkeletalMeshEditionInterface.h
--- a/Engine/Plugins/Experimental/MeshModelingToolsetExp/Source/MeshModelingToolsEditorOnlyExp/Public/SkeletalMesh/SkeletalMeshEditionInterface.h
+++ b/Engine/Plugins/Experimental/MeshModelingToolsetExp/Source/MeshModelingToolsEditorOnlyExp/Public/SkeletalMesh/SkeletalMeshEditionInterface.h
@@ -16,6 +16,7 @@ class FSkeletalMeshToolNotifier;
 class UTransformProxy;
 class UInteractiveToolManager;
 class IGizmoStateTarget;
+class USkeletonModifier;
 struct FInputDeviceRay;
 
 /**
@@ -44,6 +45,8 @@ public:
 	
 	void BindTo(TSharedPtr<ISkeletalMeshEditorBinding> InBinding);
 	void Unbind();
+
+	virtual TWeakObjectPtr<USkeletonModifier> GetModifier() const;
 	
 protected:
 	virtual void HandleSkeletalMeshModified(const TArray<FName>& BoneNames, const ESkeletalMeshNotifyType InNotifyType) = 0;
diff --git a/Engine/Plugins/Experimental/MeshModelingToolsetExp/Source/MeshModelingToolsEditorOnlyExp/Public/SkeletalMesh/SkeletonEditingTool.h b/Engine/Plugins/Experimental/MeshModelingToolsetExp/Source/MeshModelingToolsEditorOnlyExp/Public/SkeletalMesh/SkeletonEditingTool.h
--- a/Engine/Plugins/Experimental/MeshModelingToolsetExp/Source/MeshModelingToolsEditorOnlyExp/Public/SkeletalMesh/SkeletonEditingTool.h
+++ b/Engine/Plugins/Experimental/MeshModelingToolsetExp/Source/MeshModelingToolsEditorOnlyExp/Public/SkeletalMesh/SkeletonEditingTool.h
@@ -144,6 +144,9 @@ public:
 	// IModifierToggleBehaviorTarget overrides
 	virtual void OnUpdateModifierState(int ModifierID, bool bIsOn) override;
 
+	// ISkeletalMeshEditionInterface overrides
+	virtual TWeakObjectPtr<USkeletonModifier> GetModifier() const override;
+	
 protected:
 
 	// UInteractiveTool overrides
diff --git a/Engine/Source/Developer/SkeletalMeshUtilitiesCommon/Public/SkeletalMeshNotifier.h b/Engine/Source/Developer/SkeletalMeshUtilitiesCommon/Public/SkeletalMeshNotifier.h
--- a/Engine/Source/Developer/SkeletalMeshUtilitiesCommon/Public/SkeletalMeshNotifier.h
+++ b/Engine/Source/Developer/SkeletalMeshUtilitiesCommon/Public/SkeletalMeshNotifier.h
@@ -15,7 +15,8 @@ enum class ESkeletalMeshNotifyType
 	BonesRemoved,
 	BonesMoved,
 	BonesSelected,
-	BonesRenamed
+	BonesRenamed,
+	HierarchyChanged
 };
 
 // A delegate for monitoring to skeletal mesh global notifications.
diff --git a/Engine/Source/Editor/SkeletalMeshEditor/Private/SkeletalMeshEditor.cpp b/Engine/Source/Editor/SkeletalMeshEditor/Private/SkeletalMeshEditor.cpp
--- a/Engine/Source/Editor/SkeletalMeshEditor/Private/SkeletalMeshEditor.cpp
+++ b/Engine/Source/Editor/SkeletalMeshEditor/Private/SkeletalMeshEditor.cpp
@@ -1674,6 +1674,9 @@ void FSkeletalMeshEditorNotifier::HandleNotification(const TArray<FName>& BoneNa
 			SkeletonTree->SetSelectedBone(BoneNames[0], ESelectInfo::Direct);
 		}
 		break;
+	case ESkeletalMeshNotifyType::HierarchyChanged:
+		SkeletonTree->Refresh();
+		break;
 	}
 }
 
diff --git a/Engine/Source/Editor/SkeletonEditor/Private/SSkeletonTreeRow.cpp b/Engine/Source/Editor/SkeletonEditor/Private/SSkeletonTreeRow.cpp
--- a/Engine/Source/Editor/SkeletonEditor/Private/SSkeletonTreeRow.cpp
+++ b/Engine/Source/Editor/SkeletonEditor/Private/SSkeletonTreeRow.cpp
@@ -31,7 +31,10 @@ void SSkeletonTreeRow::Construct( const FArguments& InArgs, const TSharedRef<STa
 
 	check( Item.IsValid() );
 
-	SMultiColumnTableRow< TSharedPtr<ISkeletonTreeItem> >::Construct( FSuperRowType::FArguments(), InOwnerTableView );
+	const FSuperRowType::FArguments Args = FSuperRowType::FArguments()
+		.Style( FAppStyle::Get(), "TableView.AlternatingRow" );
+	
+	SMultiColumnTableRow< TSharedPtr<ISkeletonTreeItem> >::Construct( Args, InOwnerTableView );
 }
 
 void SSkeletonTreeRow::ConstructChildren(ETableViewMode::Type InOwnerTableMode, const TAttribute<FMargin>& InPadding, const TSharedRef<SWidget>& InContent)
