fixed Linux compilation error

diff --git a/Engine/Plugins/Experimental/Animation/SkeletalMeshModelingTools/Source/SkeletalMeshModelingTools/Private/SReferenceSkeletonTree.cpp b/Engine/Plugins/Experimental/Animation/SkeletalMeshModelingTools/Source/SkeletalMeshModelingTools/Private/SReferenceSkeletonTree.cpp
--- a/Engine/Plugins/Experimental/Animation/SkeletalMeshModelingTools/Source/SkeletalMeshModelingTools/Private/SReferenceSkeletonTree.cpp
+++ b/Engine/Plugins/Experimental/Animation/SkeletalMeshModelingTools/Source/SkeletalMeshModelingTools/Private/SReferenceSkeletonTree.cpp
@@ -685,7 +685,7 @@ void SReferenceSkeletonTree::RefreshTreeView(bool IsInitialSetup /*=false*/)
 	}
 
 	// expand all elements upon the initial construction of the tree
-	for (const TSharedPtr<FBoneElement> RootElement : RootElements)
+	for (const TSharedPtr<FBoneElement>& RootElement : RootElements)
 	{
 		SetExpansionRecursive(RootElement, false, true);
 	}
