Skeleton Editing Tool details customization
    - todo, better stick to the design in a second version

diff --git a/Engine/Plugins/Experimental/Animation/SkeletalMeshModelingTools/Source/SkeletalMeshModelingTools/Private/DetailCustomization/SkeletonEditingToolPropertyCustomizations.cpp b/Engine/Plugins/Experimental/Animation/SkeletalMeshModelingTools/Source/SkeletalMeshModelingTools/Private/DetailCustomization/SkeletonEditingToolPropertyCustomizations.cpp
--- a/Engine/Plugins/Experimental/Animation/SkeletalMeshModelingTools/Source/SkeletalMeshModelingTools/Private/DetailCustomization/SkeletonEditingToolPropertyCustomizations.cpp
+++ b/Engine/Plugins/Experimental/Animation/SkeletalMeshModelingTools/Source/SkeletalMeshModelingTools/Private/DetailCustomization/SkeletonEditingToolPropertyCustomizations.cpp
@@ -7,6 +7,7 @@
 #include "SkeletalMesh/SkeletonEditingTool.h"
 #include "Widgets/SBoxPanel.h"
 #include "Widgets/Input/SEditableTextBox.h"
+#include "Widgets/Input/SSegmentedControl.h"
 #include "HAL/PlatformApplicationMisc.h"
 #include "InteractiveToolManager.h"
 
@@ -18,32 +19,81 @@ TSharedRef<IDetailCustomization> FSkeletonEditingPropertiesDetailCustomization::
 	return MakeShareable(new FSkeletonEditingPropertiesDetailCustomization);
 }
 
-void FSkeletonEditingPropertiesDetailCustomization::CustomizeDetails(IDetailLayoutBuilder& DetailBuilder)
+TAttribute<EVisibility> ISkeletonEditingPropertiesDetailCustomization::GetCreationVisibility() const
 {
-	RelativeArray.Init(true, 4);
-	
-	TArray<TWeakObjectPtr<UObject>> ObjectsBeingCustomized;
-	DetailBuilder.GetObjectsBeingCustomized(ObjectsBeingCustomized);
-	if (ObjectsBeingCustomized.Num() > 1)
+	return TAttribute<EVisibility>::CreateLambda([this]()
 	{
-		return;
+		if (!Tool.IsValid())
+		{
+			return EVisibility::Hidden;
+		}
+		return Tool->GetOperation() == EEditingOperation::Create ? EVisibility::Visible : EVisibility::Hidden;
+	});
+}
+
+TAttribute<EVisibility> ISkeletonEditingPropertiesDetailCustomization::GetEditionVisibility() const
+{
+	return TAttribute<EVisibility>::CreateLambda([this]()
+	{
+		if (!Tool.IsValid())
+		{
+			return EVisibility::Hidden;
+		}
+		return Tool->GetOperation() == EEditingOperation::Create ? EVisibility::Hidden : EVisibility::Visible;
+	});
+}
+
+TAttribute<bool> ISkeletonEditingPropertiesDetailCustomization::IsEnabledBySelection() const
+{
+	return TAttribute<bool>::CreateLambda([this]()
+	{
+		if (!Tool.IsValid())
+		{
+			return false;
+		}
+		return !Tool->GetSelection().IsEmpty() && Tool->GetOperation() != EEditingOperation::Parent;
+	});
+}
+
+void ISkeletonEditingPropertiesDetailCustomization::UpdateProperties(
+	IDetailLayoutBuilder& DetailBuilder,
+	const TAttribute<EVisibility>& InVisibility,
+	const TAttribute<bool>& InEnabled) const
+{
+	for (const FName& PropName: GetProperties())
+ 
+ 
+ 
+ 
+	{
+		const TSharedRef<IPropertyHandle> Handle = DetailBuilder.GetProperty(PropName);
+		if (IDetailPropertyRow* Row = DetailBuilder.EditDefaultProperty(Handle))
+		{
+			Row->Visibility(InVisibility);
+			if (InEnabled.IsSet())
+			{
+				Row->IsEnabled(InEnabled);
+			}
+		}
 	}
+}
 
-	USkeletonEditingProperties* Properties = CastChecked<USkeletonEditingProperties>(ObjectsBeingCustomized[0]);
-	Tool = Properties ? Properties->ParentTool : nullptr;
-	if (!Tool.IsValid())
+void FSkeletonEditingPropertiesDetailCustomization::CustomizeDetails(IDetailLayoutBuilder& DetailBuilder)
+{
+	const TWeakObjectPtr<USkeletonEditingProperties> Properties = GetCustomizedObject<USkeletonEditingProperties>(DetailBuilder);
+	if (!GetParentTool(Properties).IsValid())
 	{
 		return;
 	}
-
+	
+	RelativeArray.Init(true, 4);
+	
 	IDetailCategoryBuilder& CategoryBuilder = DetailBuilder.EditCategory(TEXT("Details"));
-
 	const TSharedRef<IPropertyHandle> NodePropHandle = DetailBuilder.GetProperty(GET_MEMBER_NAME_CHECKED(USkeletonEditingProperties, Name));
 	if (NodePropHandle->IsValidHandle())
 	{
 		TSharedRef<SWidget> ValueWidget = NodePropHandle->CreatePropertyValueWidget();
 		ValueWidget->SetEnabled(TAttribute<bool>::Create([this]{ return !Tool->GetSelection().IsEmpty(); }));
-
 		static const FText MultiValues = LOCTEXT("MultipleValues", "Multiple Values");
 		
 		DetailBuilder.AddPropertyToCategory(NodePropHandle)
@@ -94,7 +144,6 @@ void FSkeletonEditingPropertiesDetailCustomization::CustomizeDetails(IDetailLayo
 		.Font(IDetailLayoutBuilder::GetDetailFont())
 		.DisplayScaleLock(false)
 		.AllowEditRotationRepresentation(false);
-
 	// relative/world
 	TransformWidgetArgs
 	.OnGetIsComponentRelative_Lambda( [this](ESlateTransformComponent::Type InComponent)
@@ -111,10 +160,8 @@ void FSkeletonEditingPropertiesDetailCustomization::CustomizeDetails(IDetailLayo
 		
 	// set bones transforms
 	CustomizeValueSet(TransformWidgetArgs);
-
 	// copy/paste values
 	CustomizeClipboard(TransformWidgetArgs);
-
 	// enabled
 	TransformWidgetArgs
 	.IsEnabled_Lambda([this]()
@@ -127,6 +174,24 @@ void FSkeletonEditingPropertiesDetailCustomization::CustomizeDetails(IDetailLayo
 		LOCTEXT("ReferenceTransform", "Transform"),
 		LOCTEXT("ReferenceBoneTransformTooltip", "The reference transform of the bone"), 
 		TransformWidgetArgs);
+	
+	const TSharedRef<IPropertyHandle> Handle = DetailBuilder.GetProperty(GET_MEMBER_NAME_CHECKED(USkeletonEditingProperties, bUpdateChildren));
+	if (IDetailPropertyRow* Row = DetailBuilder.EditDefaultProperty(Handle))
+	{
+		Row->IsEnabled(IsEnabledBySelection());
+	}
+}
+
+const TArray<FName>& FSkeletonEditingPropertiesDetailCustomization::GetProperties() const
+{
+	static const TArray< FName > Properties(
+	{
+		GET_MEMBER_NAME_CHECKED(USkeletonEditingProperties, Name),
+		GET_MEMBER_NAME_CHECKED(USkeletonEditingProperties, bUpdateChildren),
+		GET_MEMBER_NAME_CHECKED(USkeletonEditingProperties, AxisLength),
+		GET_MEMBER_NAME_CHECKED(USkeletonEditingProperties, AxisThickness),
+	});
+	return Properties;
 }
 
 void FSkeletonEditingPropertiesDetailCustomization::CustomizeValueGet(SAdvancedTransformInputBox<FTransform>::FArguments& InOutArgs)
@@ -480,4 +545,282 @@ void FSkeletonEditingPropertiesDetailCustomization::CustomizeClipboard(SAdvanced
 	});
 }
 
-#undef LOCTEXT_NAMESPACE
+TSharedRef<IDetailCustomization> FMirroringPropertiesDetailCustomization::MakeInstance()
+{
+	return MakeShareable(new FMirroringPropertiesDetailCustomization);
+}
+
+const TArray<FName>& FMirroringPropertiesDetailCustomization::GetProperties() const
+{
+	static const TArray< FName > Properties(
+	{
+		GET_MEMBER_NAME_CHECKED(UMirroringProperties, Options)
+	});
+	return Properties;
+}
+
+void FMirroringPropertiesDetailCustomization::CustomizeDetails(IDetailLayoutBuilder& DetailBuilder)
+{
+	TWeakObjectPtr<UMirroringProperties> Properties = GetCustomizedObject<UMirroringProperties>(DetailBuilder);
+	if (!GetParentTool(Properties).IsValid())
+	{
+		return;
+	}
+	const TAttribute<EVisibility> EditionVisibility = GetEditionVisibility();
+	
+	UpdateProperties(DetailBuilder, EditionVisibility, IsEnabledBySelection());
+	
+	static const FName MirrorCategoryName("Mirror");
+	IDetailCategoryBuilder& MirrorCategory = DetailBuilder.EditCategory(MirrorCategoryName, FText::GetEmpty(), ECategoryPriority::Important);
+	
+	MirrorCategory.AddCustomRow(LOCTEXT("MirrorButtonRow", "Mirror"), false)
+	[
+		SNew(SHorizontalBox)
+		+SHorizontalBox::Slot()
+		.Padding(2.f, 4.f)
+		[
+			SNew(SButton)
+			.HAlign(HAlign_Center)
+			.VAlign(VAlign_Center)
+			.Text(LOCTEXT("MirrorButtonLabel", "Mirror"))
+			.OnClicked_Lambda([Properties]()
+			{
+				Properties->MirrorBones();
+				return FReply::Handled();
+			})
+		]
+	]
+	.Visibility(EditionVisibility)
+	.IsEnabled(IsEnabledBySelection());
+}
+
+TSharedRef<IDetailCustomization> FOrientingPropertiesDetailCustomization::MakeInstance()
+{
+	return MakeShareable(new FOrientingPropertiesDetailCustomization);
+}
+
+const TArray<FName>& FOrientingPropertiesDetailCustomization::GetProperties() const
+{
+	static const TArray< FName > Properties(
+	{
+		GET_MEMBER_NAME_CHECKED(UOrientingProperties, Options),
+		GET_MEMBER_NAME_CHECKED(UOrientingProperties, bAutoOrient)
+	});
+	return Properties;
+}
+
+void FOrientingPropertiesDetailCustomization::CustomizeDetails(IDetailLayoutBuilder& DetailBuilder)
+{
+	TWeakObjectPtr<UOrientingProperties> Properties = GetCustomizedObject<UOrientingProperties>(DetailBuilder);
+	if (!GetParentTool(Properties).IsValid())
+	{
+		return;
+	}
+	const TAttribute<EVisibility> EditionVisibility = GetEditionVisibility();
+	
+	UpdateProperties(DetailBuilder, EditionVisibility, IsEnabledBySelection());
+	static const FName OrientCategoryName("Orient");
+	IDetailCategoryBuilder& MirrorCategory = DetailBuilder.EditCategory(OrientCategoryName, FText::GetEmpty(), ECategoryPriority::Important);
+	MirrorCategory.AddCustomRow(LOCTEXT("OrientButtonRow", "Orient"), false)
+	[
+		SNew(SHorizontalBox)
+		+SHorizontalBox::Slot()
+		.Padding(2.f, 4.f)
+		[
+			SNew(SButton)
+			.HAlign(HAlign_Center)
+			.VAlign(VAlign_Center)
+			.Text(LOCTEXT("OrientButtonLabel", "Orient"))
+			.OnClicked_Lambda([Properties]()
+			{
+				Properties->OrientBones();
+				return FReply::Handled();
+			})
+		]
+	]
+	.Visibility(EditionVisibility)
+	.IsEnabled(IsEnabledBySelection());
+}
+
+TSharedRef<IDetailCustomization> FProjectionPropertiesDetailCustomization::MakeInstance()
+{
+	return MakeShareable(new FProjectionPropertiesDetailCustomization);
+}
+
+const TArray<FName>& FProjectionPropertiesDetailCustomization::GetProperties() const
+{
+	static const TArray< FName > Properties(
+	{
+		GET_MEMBER_NAME_CHECKED(UProjectionProperties, ProjectionType)
+	});
+	return Properties;
+}
+
+void FProjectionPropertiesDetailCustomization::CustomizeDetails(IDetailLayoutBuilder& DetailBuilder)
+{
+	const TWeakObjectPtr<UProjectionProperties> Properties = GetCustomizedObject<UProjectionProperties>(DetailBuilder);
+	if (!GetParentTool(Properties).IsValid())
+	{
+		return;
+	}
+	UpdateProperties(DetailBuilder, GetCreationVisibility());
+}
+
+namespace SkeletonEditingCustomizationLocals
+{
+	enum EOperationType
+	{
+		Add,
+		Edit
+	};
+}
+
+TSharedRef<IDetailCustomization> FSkeletonEditingToolDetailCustomization::MakeInstance()
+{
+	return MakeShareable(new FSkeletonEditingToolDetailCustomization);
+}
+
+void FSkeletonEditingToolDetailCustomization::CustomizeDetails(IDetailLayoutBuilder& DetailBuilder)
+{
+	TArray<TWeakObjectPtr<UObject>> ObjectsBeingCustomized;
+	DetailBuilder.GetObjectsBeingCustomized(ObjectsBeingCustomized);
+	
+	TWeakObjectPtr<USkeletonEditingTool> Tool = ObjectsBeingCustomized.Num() == 1 ?
+		CastChecked<USkeletonEditingTool>(ObjectsBeingCustomized[0]) : nullptr;
+	
+	if (!Tool.IsValid())
+	{
+		return;
+	}
+	
+	auto IsCreateEnabled = [Tool]()
+	{
+		return Tool->GetOperation() == EEditingOperation::Create;
+	};
+	using namespace SkeletonEditingCustomizationLocals;
+	
+	IDetailCategoryBuilder& ActionCategory = DetailBuilder.EditCategory("Action", FText::GetEmpty(), ECategoryPriority::Important);
+	ActionCategory.AddCustomRow(LOCTEXT("ActionCategory", "Action"), false)
+	[
+		SNew(SBox)
+		.Padding(2.0f)
+		[
+			SNew(SSegmentedControl<EOperationType>)
+			
+			.Value_Lambda([IsCreateEnabled]()
+			{
+				return IsCreateEnabled() ? Add : Edit;
+			})
+			.OnValueChanged_Lambda([Tool](EOperationType Mode)
+			{
+				const EEditingOperation Operation = Mode == Add ? EEditingOperation::Create : EEditingOperation::Select;
+				Tool->SetOperation(Operation);
+			})
+			+ SSegmentedControl<EOperationType>::Slot(Add)
+			.Text(LOCTEXT("AddMode", "Add"))
+			.ToolTip(LOCTEXT("AddModeTooltip", "Create new bones. (N)"))
+			
+			+ SSegmentedControl<EOperationType>::Slot(Edit)
+			.Text(LOCTEXT("EditMode", "Edit"))
+			.ToolTip(LOCTEXT("EditModeTooltip", "Edit current bone(s) selection. (Esc)"))
+		]
+	];
+	const TAttribute<EVisibility> EditionVisibility = TAttribute<EVisibility>::CreateLambda([Tool]()
+	{
+		if (!Tool.IsValid())
+		{
+			return EVisibility::Hidden;
+		}
+		return Tool->GetOperation() == EEditingOperation::Create ? EVisibility::Hidden : EVisibility::Visible;
+	});
+	
+	const TAttribute<bool> IsEnabledBySelection = TAttribute<bool>::CreateLambda([Tool]()
+	{
+		return Tool.IsValid() ? !Tool->GetSelection().IsEmpty() : false;
+	});
+	FSegmentedControlStyle SegmentedControlStyle = FAppStyle::Get().GetWidgetStyle<FSegmentedControlStyle>("SegmentedControl");
+	
+	ActionCategory.AddCustomRow(LOCTEXT("EditCategory", "Edit"), false)
+	[
+		SNew(SHorizontalBox)
+		+SHorizontalBox::Slot()
+		.Padding(2.f, 4.f)
+		[
+			SNew(SCheckBox)
+			.Style(FAppStyle::Get(), "SegmentedCombo.ButtonOnly")
+			.OnCheckStateChanged_Lambda([Tool](ECheckBoxState InCheckState)
+			{
+				const EEditingOperation Operation = InCheckState == ECheckBoxState::Checked ?
+					EEditingOperation::Parent : EEditingOperation::Select;
+				Tool->SetOperation(Operation);
+			})
+			.IsChecked_Lambda([Tool]()
+			{
+				return Tool->GetOperation() == EEditingOperation::Parent ? ECheckBoxState::Checked : ECheckBoxState::Unchecked;
+			})
+			.IsEnabled(IsEnabledBySelection)
+			[
+				SNew(SOverlay)
+				
+				+ SOverlay::Slot()
+				.VAlign(VAlign_Center)
+				.HAlign(HAlign_Center)
+				[
+					SNew(STextBlock)
+					.Text(LOCTEXT("EditParent", "Reparent"))
+					.ToolTipText(LOCTEXT("EditParentTooltip", "Click on a bone to be set as the new parent. (B)"))
+					.TextStyle(&FAppStyle::Get().GetWidgetStyle<FTextBlockStyle>("ButtonText"))
+					.ColorAndOpacity(FLinearColor::White)
+				]
+			]
+		]
+		+SHorizontalBox::Slot()
+		.Padding(2.f, 4.f)
+		[
+			SNew(SButton)
+			.HAlign(HAlign_Center)
+			.VAlign(VAlign_Center)
+			.Text(LOCTEXT("EditDisconnect", "Disconnect"))
+			.ToolTipText(LOCTEXT("EditDisconnectTooltip", "Unparent current selection. (SHIFT+P)"))
+			.OnClicked_Lambda([Tool]()
+			{
+				Tool->UnParentBones();
+				return FReply::Handled();
+			})
+			.IsEnabled_Lambda([Tool]()
+			{
+				if (!Tool.IsValid())
+				{
+					return false;
+				}
+				return !Tool->GetSelection().IsEmpty() && Tool->GetOperation() != EEditingOperation::Parent;
+			})
+		]
+		
+		+SHorizontalBox::Slot()
+		.Padding(2.f, 4.f)
+		[
+			SNew(SButton)
+			.HAlign(HAlign_Center)
+			.VAlign(VAlign_Center)
+			.Text(LOCTEXT("EditRemove", "Remove"))
+			.ToolTipText(LOCTEXT("EditRemoveTooltip", "Remove current selection. (Delete)"))
+			.OnClicked_Lambda([Tool]()
+			{
+				Tool->RemoveBones();
+				return FReply::Handled();
+			})
+			.IsEnabled_Lambda([Tool]()
+			{
+				if (!Tool.IsValid())
+				{
+					return false;
+				}
+				return !Tool->GetSelection().IsEmpty() && Tool->GetOperation() != EEditingOperation::Parent;
+			})
+		]
+	]
+	.Visibility(EditionVisibility);
+}
+
+#undef LOCTEXT_NAMESPACE
\ No newline at end of file
diff --git a/Engine/Plugins/Experimental/Animation/SkeletalMeshModelingTools/Source/SkeletalMeshModelingTools/Private/DetailCustomization/SkeletonEditingToolPropertyCustomizations.h b/Engine/Plugins/Experimental/Animation/SkeletalMeshModelingTools/Source/SkeletalMeshModelingTools/Private/DetailCustomization/SkeletonEditingToolPropertyCustomizations.h
--- a/Engine/Plugins/Experimental/Animation/SkeletalMeshModelingTools/Source/SkeletalMeshModelingTools/Private/DetailCustomization/SkeletonEditingToolPropertyCustomizations.h
+++ b/Engine/Plugins/Experimental/Animation/SkeletalMeshModelingTools/Source/SkeletalMeshModelingTools/Private/DetailCustomization/SkeletonEditingToolPropertyCustomizations.h
@@ -9,20 +9,118 @@
 
 class IDetailLayoutBuilder;
 
-class FSkeletonEditingPropertiesDetailCustomization : public IDetailCustomization
+/**
+ * FSkeletonEditingToolDetailCustomization
+ */
+
+class FSkeletonEditingToolDetailCustomization : public IDetailCustomization
 {
 public:
 	static TSharedRef<IDetailCustomization> MakeInstance();
 	virtual void CustomizeDetails(IDetailLayoutBuilder& DetailBuilder) override;
+};
+
+/**
+ * ISkeletonEditingPropertiesDetailCustomization
+ */
+
+class ISkeletonEditingPropertiesDetailCustomization : public IDetailCustomization
+{
+protected:
+	template< typename T >
+	static TWeakObjectPtr<T> GetCustomizedObject(const IDetailLayoutBuilder& InDetailBuilder)
+	{
+		TArray<TWeakObjectPtr<UObject>> ObjectsBeingCustomized;
+		InDetailBuilder.GetObjectsBeingCustomized(ObjectsBeingCustomized);
+		if (ObjectsBeingCustomized.Num() != 1)
+		{
+			return nullptr;
+		}
+		return CastChecked<T>(ObjectsBeingCustomized[0]);
+	}
 	
-private:
+	template<typename TProperty>
+	TWeakObjectPtr<USkeletonEditingTool> GetParentTool(const TWeakObjectPtr<TProperty>& InProperties)
+	{
+		Tool = InProperties.IsValid() ? InProperties->ParentTool : nullptr;	
+		return Tool;
+	}
+
+	TAttribute<EVisibility> GetCreationVisibility() const;
+	TAttribute<EVisibility> GetEditionVisibility() const;
+	TAttribute<bool> IsEnabledBySelection() const;
+	
+	void UpdateProperties(
+		IDetailLayoutBuilder& InDetailBuilder,
+		const TAttribute<EVisibility>& InVisibility,
+		const TAttribute<bool>& InEnabled = TAttribute<bool>()) const;
+	
+	virtual const TArray<FName>& GetProperties() const = 0;
+	
+	TWeakObjectPtr<USkeletonEditingTool> Tool;
+};
 
+/**
+ * FSkeletonEditingPropertiesDetailCustomization
+ */
+
+class FSkeletonEditingPropertiesDetailCustomization : public ISkeletonEditingPropertiesDetailCustomization
+{
+public:
+	static TSharedRef<IDetailCustomization> MakeInstance();
+	virtual void CustomizeDetails(IDetailLayoutBuilder& DetailBuilder) override;
+
+protected:
+	virtual const TArray<FName>& GetProperties() const override;
+	
+private:
 	void CustomizeValueGet(SAdvancedTransformInputBox<FTransform>::FArguments& InOutArgs);
 	void CustomizeValueSet(SAdvancedTransformInputBox<FTransform>::FArguments& InOutArgs);
 	void CustomizeClipboard(SAdvancedTransformInputBox<FTransform>::FArguments& InOutArgs);
 	
-	TWeakObjectPtr<USkeletonEditingTool> Tool;
 	TUniquePtr<SkeletonEditingTool::FRefSkeletonChange> ActiveChange;
 	TBitArray<> RelativeArray;
 };
 
+/**
+ * FMirroringPropertiesDetailCustomization
+ */
+
+class FMirroringPropertiesDetailCustomization : public ISkeletonEditingPropertiesDetailCustomization
+{
+public:
+	static TSharedRef<IDetailCustomization> MakeInstance();
+	virtual void CustomizeDetails(IDetailLayoutBuilder& DetailBuilder) override;
+
+protected:
+	virtual const TArray<FName>& GetProperties() const override;
+};
+
+/**
+ * FOrientingPropertiesDetailCustomization
+ */
+
+class FOrientingPropertiesDetailCustomization : public ISkeletonEditingPropertiesDetailCustomization
+{
+public:
+	static TSharedRef<IDetailCustomization> MakeInstance();
+	virtual void CustomizeDetails(IDetailLayoutBuilder& DetailBuilder) override;
+
+protected:
+	virtual const TArray<FName>& GetProperties() const override;
+};
+
+/**
+ * FProjectionPropertiesDetailCustomization
+ */
+
+class FProjectionPropertiesDetailCustomization : public ISkeletonEditingPropertiesDetailCustomization
+{
+public:
+	static TSharedRef<IDetailCustomization> MakeInstance();
+	virtual void CustomizeDetails(IDetailLayoutBuilder& DetailBuilder) override;
+	
+protected:
+	virtual const TArray<FName>& GetProperties() const override;
+};
+
diff --git a/Engine/Plugins/Experimental/Animation/SkeletalMeshModelingTools/Source/SkeletalMeshModelingTools/Private/SkeletalMeshModelingToolsModule.cpp b/Engine/Plugins/Experimental/Animation/SkeletalMeshModelingTools/Source/SkeletalMeshModelingTools/Private/SkeletalMeshModelingToolsModule.cpp
--- a/Engine/Plugins/Experimental/Animation/SkeletalMeshModelingTools/Source/SkeletalMeshModelingTools/Private/SkeletalMeshModelingToolsModule.cpp
+++ b/Engine/Plugins/Experimental/Animation/SkeletalMeshModelingTools/Source/SkeletalMeshModelingTools/Private/SkeletalMeshModelingToolsModule.cpp
@@ -200,7 +200,11 @@ void FSkeletalMeshModelingToolsModule::RegisterPropertyCustomizations()
 		CustomizedClasses.Add(InStructName);
 	};
 
+	RegisterDetailCustomization(USkeletonEditingTool::StaticClass()->GetFName(), &FSkeletonEditingToolDetailCustomization::MakeInstance);
 	RegisterDetailCustomization(USkeletonEditingProperties::StaticClass()->GetFName(), &FSkeletonEditingPropertiesDetailCustomization::MakeInstance);
+	RegisterDetailCustomization(UMirroringProperties::StaticClass()->GetFName(), &FMirroringPropertiesDetailCustomization::MakeInstance);
+	RegisterDetailCustomization(UOrientingProperties::StaticClass()->GetFName(), &FOrientingPropertiesDetailCustomization::MakeInstance);
+	RegisterDetailCustomization(UProjectionProperties::StaticClass()->GetFName(), &FProjectionPropertiesDetailCustomization::MakeInstance);
 
 	PropertyModule.NotifyCustomizationModuleChanged();
 }
diff --git a/Engine/Plugins/Experimental/MeshModelingToolsetExp/Source/MeshModelingToolsEditorOnlyExp/Private/SkeletalMesh/SkeletonEditingTool.cpp b/Engine/Plugins/Experimental/MeshModelingToolsetExp/Source/MeshModelingToolsEditorOnlyExp/Private/SkeletalMesh/SkeletonEditingTool.cpp
--- a/Engine/Plugins/Experimental/MeshModelingToolsetExp/Source/MeshModelingToolsEditorOnlyExp/Private/SkeletalMesh/SkeletonEditingTool.cpp
+++ b/Engine/Plugins/Experimental/MeshModelingToolsetExp/Source/MeshModelingToolsEditorOnlyExp/Private/SkeletalMesh/SkeletonEditingTool.cpp
@@ -140,16 +140,12 @@ void USkeletonEditingTool::Setup()
 		UE::ToolTarget::HideSourceObject(Target);
 	}
 
+	ToolPropertyObjects.Add(this);
+	
 	// setup properties
 	{
-		Properties = NewObject<USkeletonEditingProperties>();
-		Properties->Initialize(this);
-		Properties->Name = GetCurrentBone();
-		Properties->RestoreProperties(this);
-		AddToolPropertySource(Properties);
-
 		ProjectionProperties = NewObject<UProjectionProperties>();
-		ProjectionProperties->Initialize(PreviewMesh);
+		ProjectionProperties->Initialize(this, PreviewMesh);
 		ProjectionProperties->RestoreProperties(this);
 		AddToolPropertySource(ProjectionProperties);
 
@@ -162,6 +158,12 @@ void USkeletonEditingTool::Setup()
 		OrientingProperties->Initialize(this);
 		OrientingProperties->RestoreProperties(this);
 		AddToolPropertySource(OrientingProperties);
+
+		Properties = NewObject<USkeletonEditingProperties>();
+		Properties->Initialize(this);
+		Properties->Name = GetCurrentBone();
+		Properties->RestoreProperties(this);
+		AddToolPropertySource(Properties);
 	}
 
 	// setup drag & drop behaviour
@@ -725,6 +727,21 @@ void USkeletonEditingTool::OnUpdateModifierState(int InModifierID, bool bIsOn)
 	}
 }
 
+EEditingOperation USkeletonEditingTool::GetOperation() const
+{
+	return Operation;	
+}
+
+void USkeletonEditingTool::SetOperation(const EEditingOperation InOperation, const bool bUpdateGizmo)
+{
+	Operation = InOperation;
+
+	if (bUpdateGizmo)
+	{
+		UpdateGizmo();
+	}
+}
+
 void USkeletonEditingTool::SelectBone(const FName& InBoneName)
 {
 	TArray<FName> NewSelection = Selection;
@@ -1285,8 +1302,9 @@ void UOrientingProperties::PostEditChangeProperty(FPropertyChangedEvent &Propert
 
 #endif
 
-void UProjectionProperties::Initialize(TObjectPtr<UPreviewMesh> InPreviewMesh) 
+void UProjectionProperties::Initialize(USkeletonEditingTool* ParentToolIn, TObjectPtr<UPreviewMesh> InPreviewMesh) 
 {
+	ParentTool = ParentToolIn;
 	PreviewMesh = InPreviewMesh;
 }
 
diff --git a/Engine/Plugins/Experimental/MeshModelingToolsetExp/Source/MeshModelingToolsEditorOnlyExp/Public/SkeletalMesh/SkeletonEditingTool.h b/Engine/Plugins/Experimental/MeshModelingToolsetExp/Source/MeshModelingToolsEditorOnlyExp/Public/SkeletalMesh/SkeletonEditingTool.h
--- a/Engine/Plugins/Experimental/MeshModelingToolsetExp/Source/MeshModelingToolsEditorOnlyExp/Public/SkeletalMesh/SkeletonEditingTool.h
+++ b/Engine/Plugins/Experimental/MeshModelingToolsetExp/Source/MeshModelingToolsEditorOnlyExp/Public/SkeletalMesh/SkeletonEditingTool.h
@@ -132,6 +132,8 @@ public:
 	void RenameBones();
 	void MoveBones();
 	void OrientBones();
+	void RemoveBones();
+	void UnParentBones();
 
 	// ISkeletalMeshEditionInterface overrides
 	virtual TArray<FName> GetSelectedBones() const override;
@@ -146,6 +148,9 @@ public:
 
 	// ISkeletalMeshEditionInterface overrides
 	virtual TWeakObjectPtr<USkeletonModifier> GetModifier() const override;
+
+	EEditingOperation GetOperation() const;
+	void SetOperation(const EEditingOperation InOperation, const bool bUpdateGizmo = true);
 	
 protected:
 
@@ -157,8 +162,6 @@ protected:
 
 	// Modifier functions
 	void CreateNewBone();
-	void RemoveBones();
-	void UnParentBones();
 	void ParentBones(const FName& InParentName);
 
 	UPROPERTY()
@@ -173,9 +176,12 @@ protected:
 	UPROPERTY()
 	TObjectPtr<UMirroringProperties> MirroringProperties;
 
+public:
 	UPROPERTY()
 	TObjectPtr<UOrientingProperties> OrientingProperties;
 	
+protected:
+	
 	UPROPERTY()
 	TObjectPtr<USingleClickInputBehavior> LeftClickBehavior = nullptr;
 
@@ -262,13 +268,13 @@ public:
 	UPROPERTY()
 	FTransform Transform;
 	
-	UPROPERTY(EditAnywhere, Category = "Move")
+	UPROPERTY(EditAnywhere, Category = "Details")
 	bool bUpdateChildren = false;
 
-	UPROPERTY(EditAnywhere, Category = "Viewport Axis Settings",  meta = (ClampMin = "0.0", UIMin = "0.0"))
+	UPROPERTY(EditAnywhere, Category = "Viewport Axis Settings",  meta = (DisplayPriority = 10, ClampMin = "0.0", UIMin = "0.0"))
 	float AxisLength = 1.f;
 
-	UPROPERTY(EditAnywhere, Category = "Viewport Axis Settings",  meta = (ClampMin = "0.0", UIMin = "0.0"))
+	UPROPERTY(EditAnywhere, Category = "Viewport Axis Settings",  meta = (DisplayPriority = 10, ClampMin = "0.0", UIMin = "0.0"))
 	float AxisThickness = 0.f;
 
 #if WITH_EDITOR
@@ -300,7 +306,7 @@ class MESHMODELINGTOOLSEDITORONLYEXP_API UProjectionProperties: public UInteract
 	GENERATED_BODY()
 
 public:
-	void Initialize(TObjectPtr<UPreviewMesh> PreviewMesh);
+	void Initialize(USkeletonEditingTool* ParentToolIn, TObjectPtr<UPreviewMesh> PreviewMesh);
 
 	void UpdatePlane(const UGizmoViewContext& InViewContext, const FVector& InOrigin);
 	bool GetProjectionPoint(const FInputDeviceRay& InRay, FVector& OutHitPoint) const;
@@ -309,6 +315,8 @@ public:
 	EProjectionType ProjectionType = EProjectionType::WithinMesh;
 
 	FViewCameraState CameraState;
+
+	TWeakObjectPtr<USkeletonEditingTool> ParentTool = nullptr;
 	
 private:
 	TWeakObjectPtr<UPreviewMesh> PreviewMesh = nullptr;
@@ -332,13 +340,11 @@ class MESHMODELINGTOOLSEDITORONLYEXP_API UMirroringProperties: public UInteracti
 public:
 	void Initialize(USkeletonEditingTool* ParentToolIn);
 
-	UFUNCTION(CallInEditor, Category = "Mirror")
 	void MirrorBones();
 
 	UPROPERTY(EditAnywhere, Category = "Mirror")
 	FMirrorOptions Options;
 
-private:
 	TWeakObjectPtr<USkeletonEditingTool> ParentTool = nullptr;
 };
 
@@ -354,19 +360,17 @@ class MESHMODELINGTOOLSEDITORONLYEXP_API UOrientingProperties: public UInteracti
 public:
 	void Initialize(USkeletonEditingTool* ParentToolIn);
 
-	UFUNCTION(CallInEditor, Category = "Orient")
 	void OrientBones();
-	
-	UPROPERTY(EditAnywhere, Category = "Orient")
-	FOrientOptions Options;
 
 	UPROPERTY(EditAnywhere, Category = "Orient")
 	bool bAutoOrient = false;
+	
+	UPROPERTY(EditAnywhere, Category = "Orient")
+	FOrientOptions Options;
 
 #if WITH_EDITOR
 	virtual void PostEditChangeProperty(FPropertyChangedEvent &PropertyChangedEvent) override;
 #endif
 	
-private:
 	TWeakObjectPtr<USkeletonEditingTool> ParentTool = nullptr;
 };
\ No newline at end of file
