Skelmesh Converter: Ignore stored bind bone name when creating a new skeleton, make the dialog modeless, property retain/clear binding bone when switching between skeletons/skelmeshes.

#rnx

diff --git a/Engine/Plugins/Experimental/Animation/SkeletalMeshModelingTools/Source/SkeletalMeshModelingTools/Private/SkeletalMeshModelingToolsMeshConverter.cpp b/Engine/Plugins/Experimental/Animation/SkeletalMeshModelingTools/Source/SkeletalMeshModelingTools/Private/SkeletalMeshModelingToolsMeshConverter.cpp
--- a/Engine/Plugins/Experimental/Animation/SkeletalMeshModelingTools/Source/SkeletalMeshModelingTools/Private/SkeletalMeshModelingToolsMeshConverter.cpp
+++ b/Engine/Plugins/Experimental/Animation/SkeletalMeshModelingTools/Source/SkeletalMeshModelingTools/Private/SkeletalMeshModelingToolsMeshConverter.cpp
@@ -8,6 +8,7 @@
 #include "Engine/StaticMesh.h"
 #include "PropertyEditorModule.h"
 #include "StaticToSkeletalMeshConverter.h"
+#include "AssetRegistry/AssetRegistryModule.h"
 #include "Dialog/SCustomDialog.h"
 
 DEFINE_LOG_CATEGORY_STATIC(LogSkeletalMeshModelingToolsMeshConverter, Log, All)
@@ -117,63 +118,43 @@
 	if (PropertyChangedEvent.Property)
 	{
 		const FName PropertyName = PropertyChangedEvent.Property->GetFName();
-		if (PropertyName == GET_MEMBER_NAME_CHECKED(UStaticMeshToSkeletalMeshConvertOptions, Skeleton) ||
+		if (PropertyName == GET_MEMBER_NAME_CHECKED(UStaticMeshToSkeletalMeshConvertOptions, SkeletonImportOption) ||
+			PropertyName == GET_MEMBER_NAME_CHECKED(UStaticMeshToSkeletalMeshConvertOptions, Skeleton) ||
 			PropertyName == GET_MEMBER_NAME_CHECKED(UStaticMeshToSkeletalMeshConvertOptions, SkeletalMesh))
 		{
+			// If the pointed-at-bone doesn't exist in the new ref skeleton, then clear out the bone name.
+			if (!BindingBoneName.BoneName.IsNone())
+			{
+				const FReferenceSkeleton* RefSkeleton = nullptr;
+				if (SkeletonImportOption == EReferenceSkeletonImportOption::UseExistingSkeleton)
+				{
+					const USkeleton* SkeletonPtr = Cast<USkeleton>(Skeleton.ResolveObject());
+					RefSkeleton = SkeletonPtr ? &SkeletonPtr->GetReferenceSkeleton() : nullptr;  
+				}
+				else if (SkeletonImportOption == EReferenceSkeletonImportOption::UseExistingSkeletalMesh)
+				{
+					const USkeletalMesh* SkeletalMeshPtr = Cast<USkeletalMesh>(SkeletalMesh.ResolveObject());
+					RefSkeleton = SkeletalMeshPtr ? &SkeletalMeshPtr->GetRefSkeleton() : nullptr;
+				}
+
+				if (RefSkeleton && RefSkeleton->FindBoneIndex(BindingBoneName.BoneName) == INDEX_NONE)
+				{
+					BindingBoneName.Reset();
+				}
+			}
+			
 			(void)SkeletonProviderChanged.ExecuteIfBound();
 		}
 	}
 }
 
-
-static bool ConvertSingleMeshToSkeletalMeshInteractive(
+static bool ConvertSingleMeshToSkeletalMesh(
+	const UStaticMeshToSkeletalMeshConvertOptions* InOptions,
 	UStaticMesh* InStaticMesh,
 	TArray<UObject*>& OutObjectsAdded
 	)
 {
-	UStaticMeshToSkeletalMeshConvertOptions* Options = NewObject<UStaticMeshToSkeletalMeshConvertOptions>();
-	FPropertyEditorModule& PropertyEditorModule = FModuleManager::GetModuleChecked<FPropertyEditorModule>("PropertyEditor");
-
-	// If the object paths don't resolve properly, the details view will still show something, which will be invalidated
-	// later.
-	if (Cast<USkeleton>(Options->Skeleton.ResolveObject()) == nullptr)
-	{
-		Options->Skeleton.Reset();
-	}
-	if (Cast<USkeletalMesh>(Options->SkeletalMesh.ResolveObject()) == nullptr)
-	{
-		Options->SkeletalMesh.Reset();
-	}
-
-	FDetailsViewArgs DetailsViewArgs;
-	DetailsViewArgs.bAllowSearch = false;
-	DetailsViewArgs.NameAreaSettings = FDetailsViewArgs::HideNameArea;
-	TSharedRef<IDetailsView> DetailsView = PropertyEditorModule.CreateDetailView(DetailsViewArgs);
-	DetailsView->SetObject(Options);
-
-	TSharedRef<SCustomDialog> OptionsDialog = SNew(SCustomDialog)
-		.Title(LOCTEXT("OptionsDialogTitle", "Static Mesh Conversion Options"))
-		.Content()
-		[
-			DetailsView
-		]
-	.Buttons({
-		SCustomDialog::FButton(LOCTEXT("DialogButtonConvert", "Convert")),
-		SCustomDialog::FButton(LOCTEXT("DialogButtonCancel", "Cancel"))
-	});
-
-	Options->SkeletonProviderChanged.BindLambda([DetailsView]()
-	{
-		DetailsView->ForceRefresh();
-	});
-
-	if (OptionsDialog->ShowModal() != 0)
-	{
-		return false;
-	}
-
-	Options->SaveConfig();
-
+	
 	// Make these configurable.
 	// ReSharper disable CppTooWideScope
 	static const TCHAR* DefaultSkeletonSuffix = TEXT("_Skel");
@@ -182,8 +163,9 @@
 	IAssetTools& AssetTools = FAssetToolsModule::GetModule().Get();
 	FReferenceSkeleton ReferenceSkeleton;
 	USkeleton* Skeleton = nullptr;
+	FName BindBoneName;
 
-	switch(Options->SkeletonImportOption)
+	switch(InOptions->SkeletonImportOption)
 	{
 	case EReferenceSkeletonImportOption::CreateNew:
 		{
@@ -196,23 +178,26 @@
 
 			Skeleton = Cast<USkeleton>(AssetTools.CreateAsset(SkeletonName, FPackageName::GetLongPackagePath(SkeletonPackageName), USkeleton::StaticClass(), SkeletonFactory));
 			ReferenceSkeleton = Skeleton->GetReferenceSkeleton();
+			
+			OutObjectsAdded.Add(Skeleton);
 		}
 		break;
 		
 	case EReferenceSkeletonImportOption::UseExistingSkeleton:
 		{
-			Skeleton = Cast<USkeleton>(Options->Skeleton.ResolveObject());
+			Skeleton = Cast<USkeleton>(InOptions->Skeleton.ResolveObject());
 			if (!Skeleton)
 			{
 				UE_LOG(LogSkeletalMeshModelingToolsMeshConverter, Error, TEXT("No valid skeleton given as a reference. Aborting conversion."));
 				return false;
 			}
 			ReferenceSkeleton = Skeleton->GetReferenceSkeleton();
+			BindBoneName = InOptions->BindingBoneName.BoneName;
 		}
 		break;
 	case EReferenceSkeletonImportOption::UseExistingSkeletalMesh:
 		{
-			USkeletalMesh* SkeletalMesh = Cast<USkeletalMesh>(Options->SkeletalMesh.ResolveObject());
+			USkeletalMesh* SkeletalMesh = Cast<USkeletalMesh>(InOptions->SkeletalMesh.ResolveObject());
 			if (!SkeletalMesh)
 			{
 				UE_LOG(LogSkeletalMeshModelingToolsMeshConverter, Error, TEXT("No valid skeletal mesh given as reference. Aborting conversion."));
@@ -220,6 +205,7 @@
 			}
 			Skeleton = SkeletalMesh->GetSkeleton();
 			ReferenceSkeleton = SkeletalMesh->GetRefSkeleton();
+			BindBoneName = InOptions->BindingBoneName.BoneName;
 		}
 		break;
 	}
@@ -235,31 +221,41 @@
 		SkeletalMeshFactory->StaticMesh = InStaticMesh;
 		SkeletalMeshFactory->ReferenceSkeleton = ReferenceSkeleton;
 		SkeletalMeshFactory->Skeleton = Skeleton;
-		SkeletalMeshFactory->BindBoneName = Options->BindingBoneName.BoneName; 
+		SkeletalMeshFactory->BindBoneName = BindBoneName; 
 
 		SkeletalMesh = Cast<USkeletalMesh>(AssetTools.CreateAsset(SkeletalMeshName, FPackageName::GetLongPackagePath(SkeletalMeshPackageName), USkeletalMesh::StaticClass(), SkeletalMeshFactory));
 	}
 
 	if (ensure(SkeletalMesh))
 	{
-		if (Options->SkeletonImportOption == EReferenceSkeletonImportOption::CreateNew)
+		OutObjectsAdded.Add(SkeletalMesh);
+		return true;
+	}
+
+	return false;
+}
+
+static bool ConvertMultipleMeshesToSkeletalMesh(
+	UStaticMeshToSkeletalMeshConvertOptions* InOptions,
+	TArray<UStaticMesh*> InStaticMeshes,
+	TArray<UObject*>& OutObjectsAdded
+	)
+{
+	for (UStaticMesh* StaticMesh: InStaticMeshes)
+	{
+		if (!ConvertSingleMeshToSkeletalMesh(InOptions, StaticMesh, OutObjectsAdded))
 		{
-			OutObjectsAdded.Add(Skeleton);
+			return false;
 		}
-		OutObjectsAdded.Add(SkeletalMesh);
 	}
-
 	return true;
 }
 
-
-void ConvertStaticMeshToSkeletalMeshInteractive(
+void ConvertStaticMeshAssetsToSkeletalMeshesInteractive(
 	const TArray<FAssetData>& InStaticMeshAssets
 	)
 {
-	IAssetTools& AssetTools = FAssetToolsModule::GetModule().Get();
-	
-	TArray<UObject*> ObjectsAdded;
+	TArray<UStaticMesh*> MeshesToConvert;
 
 	for (const FAssetData& AssetData: InStaticMeshAssets)
 	{
@@ -269,13 +265,66 @@
 			continue;
 		}
 
-		if (!ConvertSingleMeshToSkeletalMeshInteractive(StaticMesh, ObjectsAdded))
-		{
-			break;
-		}
+		MeshesToConvert.Add(StaticMesh);
 	}
 	
-	AssetTools.SyncBrowserToAssets(ObjectsAdded);
+
+	UStaticMeshToSkeletalMeshConvertOptions* Options = NewObject<UStaticMeshToSkeletalMeshConvertOptions>();
+	FPropertyEditorModule& PropertyEditorModule = FModuleManager::GetModuleChecked<FPropertyEditorModule>("PropertyEditor");
+
+	// If the object paths don't resolve properly, the details view will still show something, which will be invalidated
+	// later.
+	if (Cast<USkeleton>(Options->Skeleton.ResolveObject()) == nullptr)
+	{
+		Options->Skeleton.Reset();
+	}
+	if (Cast<USkeletalMesh>(Options->SkeletalMesh.ResolveObject()) == nullptr)
+	{
+		Options->SkeletalMesh.Reset();
+	}
+
+	auto OnConvertLambda = [Options, MeshesToConvert]()
+	{
+		Options->SaveConfig();
+	
+		TArray<UObject*> ObjectsAdded;
+		if (ConvertMultipleMeshesToSkeletalMesh(Options, MeshesToConvert, ObjectsAdded))
+		{
+			FAssetToolsModule::GetModule().Get().SyncBrowserToAssets(ObjectsAdded);
+		}
+		else
+		{
+			for (UObject* ObjectToDelete: ObjectsAdded)
+			{
+				FAssetRegistryModule::AssetDeleted(ObjectToDelete);
+				ObjectToDelete->Rename(nullptr, GetTransientPackage(), REN_DontCreateRedirectors | REN_ForceNoResetLoaders | REN_NonTransactional);
+			}
+		}
+	};
+
+	FDetailsViewArgs DetailsViewArgs;
+	DetailsViewArgs.bAllowSearch = false;
+	DetailsViewArgs.NameAreaSettings = FDetailsViewArgs::HideNameArea;
+	TSharedRef<IDetailsView> DetailsView = PropertyEditorModule.CreateDetailView(DetailsViewArgs);
+	DetailsView->SetObject(Options, /*bForceRefresh*/true);
+
+	TSharedRef<SCustomDialog> OptionsDialog = SNew(SCustomDialog)
+		.Title(LOCTEXT("OptionsDialogTitle", "Static Mesh Conversion Options"))
+		.Content()
+		[
+			DetailsView
+		]
+	.Buttons({
+		SCustomDialog::FButton(LOCTEXT("DialogButtonConvert", "Convert"), FSimpleDelegate::CreateLambda(OnConvertLambda)),
+		SCustomDialog::FButton(LOCTEXT("DialogButtonCancel", "Cancel"))
+	});
+
+	Options->SkeletonProviderChanged.BindLambda([DetailsView]()
+	{
+		DetailsView->ForceRefresh();
+	});
+
+	OptionsDialog->Show();
 }
 
 #undef LOCTEXT_NAMESPACE
diff --git a/Engine/Plugins/Experimental/Animation/SkeletalMeshModelingTools/Source/SkeletalMeshModelingTools/Private/SkeletalMeshModelingToolsMeshConverter.h b/Engine/Plugins/Experimental/Animation/SkeletalMeshModelingTools/Source/SkeletalMeshModelingTools/Private/SkeletalMeshModelingToolsMeshConverter.h
--- a/Engine/Plugins/Experimental/Animation/SkeletalMeshModelingTools/Source/SkeletalMeshModelingTools/Private/SkeletalMeshModelingToolsMeshConverter.h
+++ b/Engine/Plugins/Experimental/Animation/SkeletalMeshModelingTools/Source/SkeletalMeshModelingTools/Private/SkeletalMeshModelingToolsMeshConverter.h
@@ -113,4 +113,4 @@ public:
 
 
 /** Converts a static mesh to a skeletal mesh, creating both a skeletal mesh and a skeleton assets that link up. */
-void ConvertStaticMeshToSkeletalMeshInteractive(const TArray<FAssetData>& InStaticMeshAssets);
+void ConvertStaticMeshAssetsToSkeletalMeshesInteractive(const TArray<FAssetData>& InStaticMeshAssets);
diff --git a/Engine/Plugins/Experimental/Animation/SkeletalMeshModelingTools/Source/SkeletalMeshModelingTools/Private/SkeletalMeshModelingToolsModule.cpp b/Engine/Plugins/Experimental/Animation/SkeletalMeshModelingTools/Source/SkeletalMeshModelingTools/Private/SkeletalMeshModelingToolsModule.cpp
--- a/Engine/Plugins/Experimental/Animation/SkeletalMeshModelingTools/Source/SkeletalMeshModelingTools/Private/SkeletalMeshModelingToolsModule.cpp
+++ b/Engine/Plugins/Experimental/Animation/SkeletalMeshModelingTools/Source/SkeletalMeshModelingTools/Private/SkeletalMeshModelingToolsModule.cpp
@@ -113,7 +113,7 @@ void FSkeletalMeshModelingToolsModule::RegisterMenusAndToolbars()
 			{
 				if (const UContentBrowserAssetContextMenuContext* Context = InMenuContext.FindContext<UContentBrowserAssetContextMenuContext>())
 				{
-					ConvertStaticMeshToSkeletalMeshInteractive(
+					ConvertStaticMeshAssetsToSkeletalMeshesInteractive(
 						Context->GetSelectedAssetsOfType(UStaticMesh::StaticClass()));
 				}
 			});
diff --git a/Engine/Source/Runtime/Engine/Public/Animation/BoneReference.h b/Engine/Source/Runtime/Engine/Public/Animation/BoneReference.h
--- a/Engine/Source/Runtime/Engine/Public/Animation/BoneReference.h
+++ b/Engine/Source/Runtime/Engine/Public/Animation/BoneReference.h
@@ -62,6 +62,12 @@ struct FBoneReference
 	// it triggers ensure in those functions
 	ENGINE_API bool Initialize(const USkeleton* Skeleton);
 
+	/** Reset this container to default state */
+	void Reset()
+	{
+		BoneName = NAME_None;
+		InvalidateCachedBoneIndex();
+	}
 
 	/** return true if it has valid set up */
 	bool HasValidSetup() const
