Ref Skeleton Tree: postpone item's parent assignation if the parent has not been created yet (parent with a higher index)

#rnx

diff --git a/Engine/Plugins/Experimental/Animation/SkeletalMeshModelingTools/Source/SkeletalMeshModelingTools/Private/SReferenceSkeletonTree.cpp b/Engine/Plugins/Experimental/Animation/SkeletalMeshModelingTools/Source/SkeletalMeshModelingTools/Private/SReferenceSkeletonTree.cpp
--- a/Engine/Plugins/Experimental/Animation/SkeletalMeshModelingTools/Source/SkeletalMeshModelingTools/Private/SReferenceSkeletonTree.cpp
+++ b/Engine/Plugins/Experimental/Animation/SkeletalMeshModelingTools/Source/SkeletalMeshModelingTools/Private/SReferenceSkeletonTree.cpp
@@ -623,6 +623,8 @@ void SReferenceSkeletonTree::RefreshTreeView(bool IsInitialSetup /*=false*/)
 		return TextFilter->TestTextFilter(FBasicStringFilterExpressionContext(StringToTest));
 	};
 
+	TMap<int32, int32> PostponedParent;
+	
 	// create all bone elements
 	for (int32 BoneIndex=0; BoneIndex < BoneInfos.Num(); ++BoneIndex)
 	{
@@ -637,10 +639,19 @@ void SReferenceSkeletonTree::RefreshTreeView(bool IsInitialSetup /*=false*/)
 		{
 			// get parent tree element
 			const FName ParentBoneName = BoneInfos[ParentIndex].Name;
-			const TSharedPtr<FBoneElement> ParentBoneTreeElement = AllElements[BoneTreeElementIndices[ParentBoneName]];
-			BoneElement->UnFilteredParent = ParentBoneTreeElement;
-		}
 
+			const int32* FoundParent = BoneTreeElementIndices.Find(ParentBoneName);
+			if (FoundParent)
+			{ // set parent info directly
+				const TSharedPtr<FBoneElement> ParentBoneTreeElement = AllElements[*FoundParent];
+				BoneElement->UnFilteredParent = ParentBoneTreeElement;
+			}
+			else
+			{ // postpone as the parent might not have been added to AllElements / BoneTreeElementIndices yet
+				PostponedParent.Emplace(BoneElementIndex, ParentIndex);
+			}
+		}
+		
 		// apply text filter to bones
 		if (!(TextFilter->GetFilterText().IsEmpty() || FilterString(BoneName.ToString())))
 		{
@@ -648,6 +659,19 @@ void SReferenceSkeletonTree::RefreshTreeView(bool IsInitialSetup /*=false*/)
 		}
 	}
 
+	// store pointer to parents that were postponed 
+	for (const auto & [BoneElementIndex, ParentIndex] : PostponedParent)
+	{
+		const FName ParentBoneName = BoneInfos[ParentIndex].Name;
+		const int32* FoundParent = BoneTreeElementIndices.Find(ParentBoneName);
+		if (ensure(FoundParent))
+		{
+			const TSharedPtr<FBoneElement> BoneElement = AllElements[BoneElementIndex];
+			const TSharedPtr<FBoneElement> ParentBoneTreeElement = AllElements[*FoundParent];
+			BoneElement->UnFilteredParent = ParentBoneTreeElement;
+		}
+	}
+	
 	// resolve parent/children pointers on all tree elements, taking into consideration the filter options
 	// (elements are parented to their nearest non-hidden/filtered parent element)
 	for (TSharedPtr<FBoneElement>& Element : AllElements)
