Skeleton Editing Tools: fixed duplicate bones bound to the wrong action

#rnx

diff --git a/Engine/Plugins/Experimental/Animation/SkeletalMeshModelingTools/Source/SkeletalMeshModelingTools/Private/SReferenceSkeletonTree.cpp b/Engine/Plugins/Experimental/Animation/SkeletalMeshModelingTools/Source/SkeletalMeshModelingTools/Private/SReferenceSkeletonTree.cpp
--- a/Engine/Plugins/Experimental/Animation/SkeletalMeshModelingTools/Source/SkeletalMeshModelingTools/Private/SReferenceSkeletonTree.cpp
+++ b/Engine/Plugins/Experimental/Animation/SkeletalMeshModelingTools/Source/SkeletalMeshModelingTools/Private/SReferenceSkeletonTree.cpp
@@ -370,7 +370,7 @@ void SReferenceSkeletonTree::BindCommands()
 		FExecuteAction::CreateSP(this, &SReferenceSkeletonTree::HandlePasteBones),
 		FCanExecuteAction::CreateSP(this, &SReferenceSkeletonTree::CanPasteBones));
 	
-	CommandList->MapAction(Commands.UnParentBone,
+	CommandList->MapAction(Commands.DuplicateBones,
 		FExecuteAction::CreateSP(this, &SReferenceSkeletonTree::HandleDuplicateBones),
 		FCanExecuteAction::CreateSP(this, &SReferenceSkeletonTree::CanDuplicateBones));
 }
