Skeleton Editing Tools:

    - undo/redo keeps selection when possible
    - avoid re-entrant notifications
    - frame tree selection when needed

diff --git a/Engine/Plugins/Experimental/Animation/SkeletalMeshModelingTools/Source/SkeletalMeshModelingTools/Private/SReferenceSkeletonTree.cpp b/Engine/Plugins/Experimental/Animation/SkeletalMeshModelingTools/Source/SkeletalMeshModelingTools/Private/SReferenceSkeletonTree.cpp
--- a/Engine/Plugins/Experimental/Animation/SkeletalMeshModelingTools/Source/SkeletalMeshModelingTools/Private/SReferenceSkeletonTree.cpp
+++ b/Engine/Plugins/Experimental/Animation/SkeletalMeshModelingTools/Source/SkeletalMeshModelingTools/Private/SReferenceSkeletonTree.cpp
@@ -521,7 +521,8 @@ void SReferenceSkeletonTree::HandlePasteBones()
 
 	static constexpr bool bRebuildAll = true;
 	RefreshTreeView(bRebuildAll);
-	SelectItemFromNames(NewBones);
+	static constexpr bool bFrameSelection = true;
+	SelectItemFromNames(NewBones, bFrameSelection);
 
 	if (Notifier.IsValid())
 	{
@@ -557,7 +558,8 @@ void SReferenceSkeletonTree::HandleDuplicateBones()
 
 	static constexpr bool bRebuildAll = true;
 	RefreshTreeView(bRebuildAll);
-	SelectItemFromNames(NewBones);
+	static constexpr bool bFrameSelection = true;
+	SelectItemFromNames(NewBones, bFrameSelection);
 
 	if (Notifier.IsValid())
 	{
@@ -585,12 +587,18 @@ void SReferenceSkeletonTree::GetSelectedBoneNames(TArray<FName>& OutSelectedBone
 	});
 }
 
-void SReferenceSkeletonTree::SelectItemFromNames(const TArray<FName>& InBoneNames)
+void SReferenceSkeletonTree::SelectItemFromNames(const TArray<FName>& InBoneNames, bool bFrameSelection)
 {
 	for (const TSharedPtr<FBoneElement>& Item : AllElements)
 	{
 		const bool bSelect = InBoneNames.Contains(Item->BoneName);
 		TreeView->SetItemSelection(Item, bSelect, ESelectInfo::Direct);
+
+		if (bFrameSelection && bSelect)
+		{
+			TreeView->RequestScrollIntoView(Item);
+			bFrameSelection = false;
+		}
 	}
 }
 
@@ -835,8 +843,13 @@ void SReferenceSkeletonTree::HandleGetChildrenForTree(
 	OutChildren = InItem.Get()->Children;
 }
 
-void SReferenceSkeletonTree::OnSelectionChanged(TSharedPtr<FBoneElement> Selection, ESelectInfo::Type SelectInfo)
+void SReferenceSkeletonTree::OnSelectionChanged(TSharedPtr<FBoneElement> InItem, ESelectInfo::Type InSelectInfo)
 {
+	if (!InItem || InSelectInfo == ESelectInfo::Direct)
+	{
+		return;
+	}
+	
 	TArray<FName> BoneNames;
 	GetSelectedBoneNames(BoneNames);
 	
@@ -1027,7 +1040,7 @@ FReferenceSkeletonWidgetNotifier::FReferenceSkeletonWidgetNotifier(TSharedRef<SR
 
 void FReferenceSkeletonWidgetNotifier::HandleNotification(const TArray<FName>& InBoneNames, const ESkeletalMeshNotifyType InNotifyType)
 {
-	if (!Tree.IsValid())
+	if (Notifying() || !Tree.IsValid())
 	{
 		return;
 	}
@@ -1044,8 +1057,11 @@ void FReferenceSkeletonWidgetNotifier::HandleNotification(const TArray<FName>& I
 		case ESkeletalMeshNotifyType::BonesMoved:
 			break;
 		case ESkeletalMeshNotifyType::BonesSelected:
-			TreePtr->SelectItemFromNames(InBoneNames);
-			break;
+			{
+				static constexpr bool bFrameSelection = true;
+				TreePtr->SelectItemFromNames(InBoneNames, bFrameSelection);
+				break;
+			}
 		case ESkeletalMeshNotifyType::BonesRenamed:
 			TreePtr->RefreshTreeView();
 			TreePtr->SelectItemFromNames(InBoneNames);
diff --git a/Engine/Plugins/Experimental/Animation/SkeletalMeshModelingTools/Source/SkeletalMeshModelingTools/Private/SReferenceSkeletonTree.h b/Engine/Plugins/Experimental/Animation/SkeletalMeshModelingTools/Source/SkeletalMeshModelingTools/Private/SReferenceSkeletonTree.h
--- a/Engine/Plugins/Experimental/Animation/SkeletalMeshModelingTools/Source/SkeletalMeshModelingTools/Private/SReferenceSkeletonTree.h
+++ b/Engine/Plugins/Experimental/Animation/SkeletalMeshModelingTools/Source/SkeletalMeshModelingTools/Private/SReferenceSkeletonTree.h
@@ -176,7 +176,7 @@ public:
 	void GetSelectedBoneNames(TArray<FName>& OutSelectedBoneNames) const;
 	TArray<TSharedPtr<FBoneElement>> GetSelectedItems() const;
 	bool HasSelectedItems() const;
-	void SelectItemFromNames(const TArray<FName>& InBoneNames);
+	void SelectItemFromNames(const TArray<FName>& InBoneNames, bool bFrameSelection = false);
 
 	ISkeletalMeshNotifier& GetNotifier();
 
@@ -215,7 +215,7 @@ private:
 	// Callbacks
 	void RefreshTreeView(bool IsInitialSetup=false);
 	void HandleGetChildrenForTree(TSharedPtr<FBoneElement> InItem, TArray<TSharedPtr<FBoneElement>>& OutChildren);
-	void OnSelectionChanged(TSharedPtr<FBoneElement> Selection, ESelectInfo::Type SelectInfo);
+	void OnSelectionChanged(TSharedPtr<FBoneElement> InItem, ESelectInfo::Type InSelectInfo);
 	
 	TSharedRef< SWidget > CreateAddNewMenu();
 	TSharedPtr< SWidget > CreateContextMenu();
diff --git a/Engine/Source/Editor/SkeletalMeshEditor/Private/SkeletalMeshEditor.cpp b/Engine/Source/Editor/SkeletalMeshEditor/Private/SkeletalMeshEditor.cpp
--- a/Engine/Source/Editor/SkeletalMeshEditor/Private/SkeletalMeshEditor.cpp
+++ b/Engine/Source/Editor/SkeletalMeshEditor/Private/SkeletalMeshEditor.cpp
@@ -1568,7 +1568,7 @@ FSkeletalMeshEditorNotifier::FSkeletalMeshEditorNotifier(TSharedRef<FSkeletalMes
 
 void FSkeletalMeshEditorNotifier::HandleNotification(const TArray<FName>& BoneNames, const ESkeletalMeshNotifyType InNotifyType)
 {
-	if (!Editor.IsValid())
+	if (Notifying() || !Editor.IsValid())
 	{
 		return;
 	}
