Skeleton Editing Tool: add reset button for transforms

#rnx

diff --git a/Engine/Plugins/Experimental/Animation/SkeletalMeshModelingTools/Source/SkeletalMeshModelingTools/Private/DetailCustomization/SkeletonEditingToolPropertyCustomizations.cpp b/Engine/Plugins/Experimental/Animation/SkeletalMeshModelingTools/Source/SkeletalMeshModelingTools/Private/DetailCustomization/SkeletonEditingToolPropertyCustomizations.cpp
--- a/Engine/Plugins/Experimental/Animation/SkeletalMeshModelingTools/Source/SkeletalMeshModelingTools/Private/DetailCustomization/SkeletonEditingToolPropertyCustomizations.cpp
+++ b/Engine/Plugins/Experimental/Animation/SkeletalMeshModelingTools/Source/SkeletalMeshModelingTools/Private/DetailCustomization/SkeletonEditingToolPropertyCustomizations.cpp
@@ -214,7 +214,7 @@ void FSkeletonEditingPropertiesDetailCustomization::CustomizeValueGet(SAdvancedT
 			  InRepresentation,
 			  InSubComponent);
 	};
-	
+
 	InOutArgs.OnGetNumericValue_Lambda( [this, GetNumericValue](
 		ESlateTransformComponent::Type InComponent,
 		ESlateRotationRepresentation::Type InRepresentation,
@@ -244,6 +244,40 @@ void FSkeletonEditingPropertiesDetailCustomization::CustomizeValueGet(SAdvancedT
 		}
 		return Value;
 	} );
+	
+	InOutArgs.DiffersFromDefault_Lambda([this](ESlateTransformComponent::Type InComponent)
+	{
+		const TArray<FName>& Bones = Tool->GetSelection();
+		for (const FName& BoneName: Bones)
+		{
+			const bool bWorld = !RelativeArray[InComponent];
+			const FTransform& Transform = Tool->GetTransform(BoneName, bWorld);
+			switch (InComponent)
+			{
+				case ESlateTransformComponent::Location:
+					if (!Transform.GetLocation().Equals(FVector::ZeroVector))
+					{
+						return true;
+					}
+					break;
+				case ESlateTransformComponent::Rotation:
+					if(!Transform.GetRotation().Equals(FQuat::Identity))
+					{
+						return true;
+					}
+					break;
+				case ESlateTransformComponent::Scale:
+					if(!Transform.GetScale3D().Equals(FVector::OneVector))
+					{
+						return true;
+					}
+					break;
+				default:
+					break;
+			}
+		}
+		return false;
+	});
 }
 
 void FSkeletonEditingPropertiesDetailCustomization::CustomizeValueSet(SAdvancedTransformInputBox<FTransform>::FArguments& InOutArgs)
@@ -298,8 +332,9 @@ void FSkeletonEditingPropertiesDetailCustomization::CustomizeValueSet(SAdvancedT
 			const bool bWorld = !RelativeArray[InComponent];
 			Tool->SetTransforms(BonesToMove, UpdatedTransforms, bWorld);
 		}
-	})
-	.OnNumericValueCommitted_Lambda([this, PrepareNumericValueChanged](
+	});
+
+	InOutArgs.OnNumericValueCommitted_Lambda([this, PrepareNumericValueChanged](
 		ESlateTransformComponent::Type InComponent,
 		ESlateRotationRepresentation::Type InRepresentation,
 		ESlateTransformSubComponent::Type InSubComponent,
@@ -308,8 +343,11 @@ void FSkeletonEditingPropertiesDetailCustomization::CustomizeValueSet(SAdvancedT
 	{
 		const TArray<FName>& Bones = Tool->GetSelection();
 
-		TArray<FName> BonesToMove; BonesToMove.Reserve(Bones.Num());
-		TArray<FTransform> UpdatedTransforms; UpdatedTransforms.Reserve(Bones.Num());
+		TArray<FName> BonesToMove;
+		BonesToMove.Reserve(Bones.Num());
+		
+		TArray<FTransform> UpdatedTransforms;
+		UpdatedTransforms.Reserve(Bones.Num());
 		
 		FTransform CurrentTransform, UpdatedTransform;
 		for (const FName& BoneName: Bones)
@@ -350,6 +388,79 @@ void FSkeletonEditingPropertiesDetailCustomization::CustomizeValueSet(SAdvancedT
 			ActiveChange.Reset();
 		}
 	});
+
+	auto PrepareNumericValueReset = [this]( const FName InBoneName, ESlateTransformComponent::Type InComponent)
+	{
+		const bool bWorld = !RelativeArray[InComponent];
+		const FTransform& InTransform = Tool->GetTransform(InBoneName, bWorld);
+		FTransform OutTransform = InTransform;
+
+		switch (InComponent)
+		{
+		case ESlateTransformComponent::Location:
+			OutTransform.SetLocation(FVector::ZeroVector);
+			break;
+		case ESlateTransformComponent::Rotation:
+			OutTransform.SetRotation(FQuat::Identity);
+			break;
+		case ESlateTransformComponent::Scale:
+			OutTransform.SetScale3D(FVector::OneVector);
+			break;
+		default:
+			break;
+		}
+		
+		return MakeTuple(InTransform, OutTransform);
+	};
+	
+	InOutArgs.OnResetToDefault_Lambda([this, PrepareNumericValueReset](ESlateTransformComponent::Type InComponent)
+	{
+		const TArray<FName>& Bones = Tool->GetSelection();
+
+		TArray<FName> BonesToMove;
+		BonesToMove.Reserve(Bones.Num());
+		
+		TArray<FTransform> UpdatedTransforms;
+		UpdatedTransforms.Reserve(Bones.Num());
+		
+		FTransform CurrentTransform, UpdatedTransform;
+		for (const FName& BoneName: Bones)
+		{
+			Tie(CurrentTransform, UpdatedTransform) = PrepareNumericValueReset(BoneName, InComponent);
+			if (!UpdatedTransform.Equals(CurrentTransform))
+			{
+				BonesToMove.Add(BoneName);
+				UpdatedTransforms.Add(UpdatedTransform);
+			}
+		}
+
+		if (!BonesToMove.IsEmpty())
+		{
+			if (!ActiveChange.IsValid())
+			{
+				ActiveChange = MakeUnique<SkeletonEditingTool::FRefSkeletonChange>(Tool.Get());
+			}
+
+			const bool bWorld = !RelativeArray[InComponent];
+			Tool->SetTransforms(BonesToMove, UpdatedTransforms, bWorld);
+		}
+
+		if (ActiveChange.IsValid())
+		{
+			// send transaction
+			if (UInteractiveToolManager* ToolManager = Tool->GetToolManager())
+			{
+				ActiveChange->StoreSkeleton(Tool.Get());
+
+				static const FText TransactionDesc = LOCTEXT("ResetNumericValue", "Reset Numeric Value");
+				ToolManager->BeginUndoTransaction(TransactionDesc);
+				ToolManager->EmitObjectChange(Tool.Get(), MoveTemp(ActiveChange), TransactionDesc);
+				ToolManager->EndUndoTransaction();
+			}
+				
+			ActiveChange.Reset();
+		}
+	});
 }
 
 namespace FSkeletonEditingToolClipboardLocals
