Skeletal Mesh Editor: better integration of gizmo manipulations across the skeletal mesh editor

#rnx

diff --git a/Engine/Plugins/Experimental/Animation/SkeletalMeshModelingTools/Source/SkeletalMeshModelingTools/Private/SkeletalMeshModelingToolsEditorMode.cpp b/Engine/Plugins/Experimental/Animation/SkeletalMeshModelingTools/Source/SkeletalMeshModelingTools/Private/SkeletalMeshModelingToolsEditorMode.cpp
--- a/Engine/Plugins/Experimental/Animation/SkeletalMeshModelingTools/Source/SkeletalMeshModelingTools/Private/SkeletalMeshModelingToolsEditorMode.cpp
+++ b/Engine/Plugins/Experimental/Animation/SkeletalMeshModelingTools/Source/SkeletalMeshModelingTools/Private/SkeletalMeshModelingToolsEditorMode.cpp
@@ -367,11 +367,24 @@ bool USkeletalMeshModelingToolsEditorMode::ComputeBoundingBoxForViewportFocus(AA
 void USkeletalMeshModelingToolsEditorMode::OnToolStarted(UInteractiveToolManager* Manager, UInteractiveTool* Tool)
 {
 	FSkeletalMeshModelingToolsActionCommands::UpdateToolCommandBinding(Tool, Toolkit->GetToolkitCommands(), false);
+
+	// deactivate SkeletonSelection when a tool is activated.
+	// each tool is responsible for activating SkeletonSelection if necessary
+	if (Owner)
+	{
+		Owner->DeactivateMode(FPersonaEditModes::SkeletonSelection);
+	}
 }
 
 void USkeletalMeshModelingToolsEditorMode::OnToolEnded(UInteractiveToolManager* Manager, UInteractiveTool* Tool)
 {
 	FSkeletalMeshModelingToolsActionCommands::UpdateToolCommandBinding(Tool, Toolkit->GetToolkitCommands(), true);
+
+	// reactivate SkeletonSelection when deactivating a tool
+	if (Owner)
+	{
+		Owner->ActivateMode(FPersonaEditModes::SkeletonSelection);
+	}
 }
 
 void USkeletalMeshModelingToolsEditorMode::SetEditorBinding(const TWeakPtr<ISkeletalMeshEditor>& InSkeletalMeshEditor)
diff --git a/Engine/Plugins/Experimental/Animation/SkeletalMeshModelingTools/Source/SkeletalMeshModelingTools/Private/SkeletalMeshModelingToolsModule.cpp b/Engine/Plugins/Experimental/Animation/SkeletalMeshModelingTools/Source/SkeletalMeshModelingTools/Private/SkeletalMeshModelingToolsModule.cpp
--- a/Engine/Plugins/Experimental/Animation/SkeletalMeshModelingTools/Source/SkeletalMeshModelingTools/Private/SkeletalMeshModelingToolsModule.cpp
+++ b/Engine/Plugins/Experimental/Animation/SkeletalMeshModelingTools/Source/SkeletalMeshModelingTools/Private/SkeletalMeshModelingToolsModule.cpp
@@ -14,6 +14,7 @@
 #include "ISkeletalMeshEditorModule.h"
 #include "ModelingToolsEditorModeStyle.h"
 #include "ModelingToolsManagerActions.h"
+#include "PersonaModule.h"
 #include "Modules/ModuleManager.h"
 #include "SkeletalMeshToolMenuContext.h"
 #include "ToolMenus.h"
@@ -168,6 +169,13 @@ void FSkeletalMeshModelingToolsModule::OnToggleEditingToolsMode(TWeakPtr<ISkelet
 		{
 			EditorModeManager.DeactivateMode(USkeletalMeshModelingToolsEditorMode::Id);
 		}
+
+		// make sure SkeletonSelection is active when toggling the mode, as they are compatible.
+		// it will be deactivated later when entering a tool 
+		if (!EditorModeManager.IsModeActive(FPersonaEditModes::SkeletonSelection))
+		{
+			EditorModeManager.ActivateMode(FPersonaEditModes::SkeletonSelection);
+		}
 	}
 }
 
