Skeleton Editing Tool: component selection improvements
    - Y to cycle between components
    - V to snap/create new bone based on component selection
    - backface on by default
    - better UI

diff --git a/Engine/Plugins/Experimental/Animation/SkeletalMeshModelingTools/Source/SkeletalMeshModelingTools/Private/DetailCustomization/SkeletonEditingToolPropertyCustomizations.cpp b/Engine/Plugins/Experimental/Animation/SkeletalMeshModelingTools/Source/SkeletalMeshModelingTools/Private/DetailCustomization/SkeletonEditingToolPropertyCustomizations.cpp
--- a/Engine/Plugins/Experimental/Animation/SkeletalMeshModelingTools/Source/SkeletalMeshModelingTools/Private/DetailCustomization/SkeletonEditingToolPropertyCustomizations.cpp
+++ b/Engine/Plugins/Experimental/Animation/SkeletalMeshModelingTools/Source/SkeletalMeshModelingTools/Private/DetailCustomization/SkeletonEditingToolPropertyCustomizations.cpp
@@ -818,89 +818,95 @@ void FSkeletonEditingToolDetailCustomization::CustomizeDetails(IDetailLayoutBuil
 	IDetailCategoryBuilder& ActionCategory = DetailBuilder.EditCategory("Action", FText::GetEmpty(), ECategoryPriority::Important);
 	ActionCategory.AddCustomRow(LOCTEXT("ActionCategory", "Action"), false)
 	[
-		SNew(SHorizontalBox)
-		+SHorizontalBox::Slot()
-		.Padding(2.f, 4.f)
+		SNew(SBorder)
+		.BorderImage(&SegmentedControlStyle.BackgroundBrush)
+		.Padding(FMargin(0,0,SegmentedControlStyle.UniformPadding.Right,0))
 		[
-			SNew(SCheckBox)
-			.Style(&SegmentedControlStyle.FirstControlStyle)
-			.ToolTipText(LOCTEXT("AddModeTooltip", "Create new bones. (N)"))
-			.OnCheckStateChanged_Lambda([this](ECheckBoxState InCheckState)
-			{
-				if (InCheckState == ECheckBoxState::Checked)
-				{
-					Tool->SetOperation(EEditingOperation::Create);
-				}
-			})
-			.IsChecked_Lambda([IsCreateEnabled]()
-			{
-				return IsCreateEnabled() ? ECheckBoxState::Checked : ECheckBoxState::Unchecked;
-			})
+			SNew(SHorizontalBox)
+			+SHorizontalBox::Slot()
 			[
-				SNew(SOverlay)
-				+ SOverlay::Slot()
-				.VAlign(VAlign_Center)
-				.HAlign(HAlign_Center)
+				SNew(SCheckBox)
+				.Clipping(EWidgetClipping::ClipToBounds)
+				.Style(&SegmentedControlStyle.FirstControlStyle)
+				.ToolTipText(LOCTEXT("AddModeTooltip", "Create new bones. (N)"))
+				.OnCheckStateChanged_Lambda([this](ECheckBoxState InCheckState)
+				{
+					if (InCheckState == ECheckBoxState::Checked)
+					{
+						Tool->SetOperation(EEditingOperation::Create);
+					}
+				})
+				.IsChecked_Lambda([IsCreateEnabled]()
+				{
+					return IsCreateEnabled() ? ECheckBoxState::Checked : ECheckBoxState::Unchecked;
+				})
 				[
-					SNew(STextBlock)
-					.Text(LOCTEXT("AddMode", "Add"))
-					.TextStyle(&FAppStyle::Get().GetWidgetStyle<FTextBlockStyle>("ButtonText"))
-					.ColorAndOpacity(FLinearColor::White)
+					SNew(SOverlay)
+					+ SOverlay::Slot()
+					.VAlign(VAlign_Center)
+					.HAlign(HAlign_Center)
+					[
+						SNew(STextBlock)
+						.Text(LOCTEXT("AddMode", "Add"))
+						.TextStyle(&FAppStyle::Get().GetWidgetStyle<FTextBlockStyle>("ButtonText"))
+						.ColorAndOpacity(FLinearColor::White)
+					]
 				]
 			]
-		]
-		+SHorizontalBox::Slot()
-		.Padding(2.f, 4.f)
-		[
-			SNew(SCheckBox)
-			.Style(&SegmentedControlStyle.LastControlStyle)
-			.ToolTipText(LOCTEXT("EditModeTooltip", "Edit current bone(s) selection. (Esc)"))
-			.OnCheckStateChanged_Lambda([this](ECheckBoxState InCheckState)
-			{
-				if (InCheckState == ECheckBoxState::Checked)
-				{
-					Tool->SetOperation(EEditingOperation::Select);
-				}
-			})
-			.IsChecked_Lambda([IsCreateEnabled]()
-			{
-				return IsCreateEnabled() ? ECheckBoxState::Unchecked : ECheckBoxState::Checked;
-			})
+			+SHorizontalBox::Slot()
 			[
-				SNew(SOverlay)
-					
-				+ SOverlay::Slot()
-				.VAlign(VAlign_Center)
-				.HAlign(HAlign_Center)
+				SNew(SCheckBox)
+				.Clipping(EWidgetClipping::ClipToBounds)
+				.Style(&SegmentedControlStyle.LastControlStyle)
+				.ToolTipText(LOCTEXT("EditModeTooltip", "Edit current bone(s) selection. (Esc)"))
+				.OnCheckStateChanged_Lambda([this](ECheckBoxState InCheckState)
+				{
+					if (InCheckState == ECheckBoxState::Checked)
+					{
+						Tool->SetOperation(EEditingOperation::Select);
+					}
+				})
+				.IsChecked_Lambda([IsCreateEnabled]()
+				{
+					return IsCreateEnabled() ? ECheckBoxState::Unchecked : ECheckBoxState::Checked;
+				})
 				[
-					SNew(STextBlock)
-					.Text(LOCTEXT("EditMode", "Edit"))
-					.TextStyle(&FAppStyle::Get().GetWidgetStyle<FTextBlockStyle>("ButtonText"))
-					.ColorAndOpacity(FLinearColor::White)
+					SNew(SOverlay)
+						
+					+ SOverlay::Slot()
+					.VAlign(VAlign_Center)
+					.HAlign(HAlign_Center)
+					[
+						SNew(STextBlock)
+						.Text(LOCTEXT("EditMode", "Edit"))
+						.TextStyle(&FAppStyle::Get().GetWidgetStyle<FTextBlockStyle>("ButtonText"))
+						.ColorAndOpacity(FLinearColor::White)
+					]
 				]
 			]
-		]
 
-		+ SHorizontalBox::Slot()
-		.Padding(8.f, 4.f)
-		.AutoWidth()
-		.HAlign(HAlign_Center)
-		.VAlign(VAlign_Center)
-		[
-			SNew(SCheckBox)
-			.Style(FAppStyle::Get(), "ToggleButtonCheckbox")
-			.ToolTipText(LOCTEXT("ComponentModeTooltip", "Select vertices, edges, and triangles to place bone. (T)"))
-			.OnCheckStateChanged_Lambda([this](ECheckBoxState)
-			{
-				Tool->Properties->bEnableComponentSelection = !Tool->Properties->bEnableComponentSelection;
-			})
-			.IsChecked_Lambda([this]()
-			{
-				return Tool->Properties->bEnableComponentSelection ? ECheckBoxState::Checked : ECheckBoxState::Unchecked;
-			})
+			+ SHorizontalBox::Slot()
+			.Padding(FMargin(4.f, 0.f, 0.f, 0.f))
+			.AutoWidth()
+			.HAlign(HAlign_Center)
+			.VAlign(VAlign_Center)
 			[
-				SNew(SImage)
-					.Image(FModelingToolsEditorModeStyle::Get()->GetBrush("ModelingToolsManagerCommands.BeginMeshSelectionTool"))
+				SNew(SCheckBox)
+				.Clipping(EWidgetClipping::ClipToBounds)
+				.Style(FAppStyle::Get(), "ToggleButtonCheckbox")
+				.ToolTipText(LOCTEXT("ComponentModeTooltip", "Select vertices, edges, and triangles to place bone. (T)"))
+				.OnCheckStateChanged_Lambda([this](ECheckBoxState)
+				{
+					Tool->Properties->bEnableComponentSelection = !Tool->Properties->bEnableComponentSelection;
+				})
+				.IsChecked_Lambda([this]()
+				{
+					return Tool->Properties->bEnableComponentSelection ? ECheckBoxState::Checked : ECheckBoxState::Unchecked;
+				})
+				[
+					SNew(SImage)
+						.Image(FModelingToolsEditorModeStyle::Get()->GetBrush("ModelingToolsManagerCommands.BeginMeshSelectionTool"))
+				]
 			]
 		]
 	];
diff --git a/Engine/Plugins/Experimental/MeshModelingToolsetExp/Source/MeshModelingToolsEditorOnlyExp/Private/SkeletalMesh/SkeletonEditingTool.cpp b/Engine/Plugins/Experimental/MeshModelingToolsetExp/Source/MeshModelingToolsEditorOnlyExp/Private/SkeletalMesh/SkeletonEditingTool.cpp
--- a/Engine/Plugins/Experimental/MeshModelingToolsetExp/Source/MeshModelingToolsEditorOnlyExp/Private/SkeletalMesh/SkeletonEditingTool.cpp
+++ b/Engine/Plugins/Experimental/MeshModelingToolsetExp/Source/MeshModelingToolsEditorOnlyExp/Private/SkeletalMesh/SkeletonEditingTool.cpp
@@ -292,8 +292,6 @@ void USkeletonEditingTool::SetupComponentsSelection()
 	// the property is disabled and managed thru customization
 	SetToolPropertySourceEnabled(SelectionMechanic->Properties, false);
 	SelectionMechanic->SetMarqueeSelectionUpdateType(EMarqueeSelectionUpdateType::OnRelease);
-	SelectionMechanic->Properties->bMarqueeIgnoreOcclusion = false;
-	SelectionMechanic->Properties->bHitBackFaces = false;
 	SelectionMechanic->Properties->bDisplayPolygroupReliantControls = false;
 
 	SelectionMechanic->Properties->bSelectVertices = true;
@@ -329,7 +327,9 @@ void USkeletonEditingTool::SetupComponentsSelection()
 	// triangles rendering
 	PreviewMesh->EnableSecondaryTriangleBuffers([this](const FDynamicMesh3* Mesh, int32 TriangleID)
 	{
-		return SelectionMechanic->GetActiveSelection().IsSelectedTriangle(Mesh, Topology.Get(), TriangleID);
+		return	Properties->bEnableComponentSelection &&
+				SelectionMechanic->Properties->bSelectFaces &&
+				SelectionMechanic->GetActiveSelection().IsSelectedTriangle(Mesh, Topology.Get(), TriangleID);
 	});
 	SelectionMechanic->OnSelectionChanged.AddWeakLambda(this, [this]()
 	{
@@ -469,6 +469,8 @@ void USkeletonEditingTool::RegisterActions(FInteractiveToolActionSet& InOutActio
 	RegisterPasteAction(InOutActionSet, GetActionId());
 	RegisterDuplicateAction(InOutActionSet, GetActionId());
 	RegisterSelectComponentsAction(InOutActionSet, GetActionId());
+	RegisterSelectionFilterCyclingAction(InOutActionSet, GetActionId());
+	RegisterSnapAction(InOutActionSet, GetActionId());
 }
 
 void USkeletonEditingTool::RegisterCreateAction(FInteractiveToolActionSet& InOutActionSet, const int32 InActionId)
@@ -644,11 +646,71 @@ void USkeletonEditingTool::RegisterSelectComponentsAction(FInteractiveToolAction
 		EModifierKey::None, EKeys::T,
 		[this]()
 		{
-			Properties->bEnableComponentSelection = !Properties->bEnableComponentSelection; 
+			Properties->bEnableComponentSelection = !Properties->bEnableComponentSelection;
+			PreviewMesh->FastNotifySecondaryTrianglesChanged();
 			UpdateGizmo();
 		});
 }
 
+void USkeletonEditingTool::RegisterSelectionFilterCyclingAction(FInteractiveToolActionSet& InOutActionSet, const int32 InActionId)
+{
+	auto UpdateSelect = [this](bool bSelectVertices, bool bSelectEdges, bool bSelectFaces)
+	{
+		const bool bChangingFace = SelectionMechanic->Properties->bSelectFaces != bSelectFaces;
+
+		SelectionMechanic->Properties->bSelectVertices = bSelectVertices;
+		SelectionMechanic->Properties->bSelectEdges = bSelectEdges;
+		SelectionMechanic->Properties->bSelectFaces = bSelectFaces;
+		SelectionMechanic->SetShowSelectableCorners(bSelectVertices);
+
+		if (bChangingFace)
+		{
+			PreviewMesh->FastNotifySecondaryTrianglesChanged();
+		}
+	};
+	
+	InOutActionSet.RegisterAction(this, InActionId,
+		TEXT("ComponentCycling"),
+		LOCTEXT("ComponentCycling", "Cycle Selection Filter"),
+		LOCTEXT("ComponentCyclingDesc", "Cycle between vertex, edge, face selection"),
+		EModifierKey::None, EKeys::Y,
+		[this, UpdateSelect]()
+		{
+			if (SelectionMechanic->Properties->bSelectVertices)
+			{
+				return UpdateSelect(false, true, false);
+			}
+			
+			if (SelectionMechanic->Properties->bSelectEdges)
+			{
+				return UpdateSelect(false, false, true);
+			}
+
+			if (SelectionMechanic->Properties->bSelectFaces)
+			{
+				return UpdateSelect(true, false, false);
+			}
+			
+			UpdateSelect(true, false, false);
+		});
+}
+
+void USkeletonEditingTool::RegisterSnapAction(FInteractiveToolActionSet& InOutActionSet, const int32 InActionId)
+{
+	InOutActionSet.RegisterAction(this, InActionId, TEXT("SnapBone"),
+		LOCTEXT("SnapBone", "Snap Bone"),
+		LOCTEXT("SnapBoneBonesDesc", "Snap Bone"),
+		EModifierKey::None, EKeys::V,
+		[this]()
+		{
+			if (Properties->bEnableComponentSelection)
+			{
+				SnapBoneToComponentSelection(Operation == EEditingOperation::Create);
+				UpdateGizmo();
+			}
+		});
+}
+
 TArray<int32> USkeletonEditingTool::GetSelectedComponents() const
 {
 	TArray<int32> IDs;
@@ -1612,7 +1674,7 @@ void USkeletonEditingTool::DrawHUD(FCanvas* Canvas, IToolsContextRenderAPI* Rend
 {
 	Super::DrawHUD(Canvas, RenderAPI);
 
-	if (SelectionMechanic)
+	if (SelectionMechanic && Properties->bEnableComponentSelection)
 	{
 		SelectionMechanic->DrawHUD(Canvas, RenderAPI);
 	}
@@ -1708,6 +1770,7 @@ void USkeletonEditingTool::OnPropertyModified(UObject* PropertySet, FProperty* P
 			}
 		}
 		SelectionMechanic->SetShowSelectableCorners(SelectionMechanic->Properties->bSelectVertices);
+		PreviewMesh->FastNotifySecondaryTrianglesChanged();
 	}
 	Super::OnPropertyModified(PropertySet, Property);
 }
diff --git a/Engine/Plugins/Experimental/MeshModelingToolsetExp/Source/MeshModelingToolsEditorOnlyExp/Public/SkeletalMesh/SkeletonEditingTool.h b/Engine/Plugins/Experimental/MeshModelingToolsetExp/Source/MeshModelingToolsEditorOnlyExp/Public/SkeletalMesh/SkeletonEditingTool.h
--- a/Engine/Plugins/Experimental/MeshModelingToolsetExp/Source/MeshModelingToolsEditorOnlyExp/Public/SkeletalMesh/SkeletonEditingTool.h
+++ b/Engine/Plugins/Experimental/MeshModelingToolsetExp/Source/MeshModelingToolsEditorOnlyExp/Public/SkeletalMesh/SkeletonEditingTool.h
@@ -274,6 +274,8 @@ private:
 	void RegisterPasteAction(FInteractiveToolActionSet& InOutActionSet, const int32 InActionId);
 	void RegisterDuplicateAction(FInteractiveToolActionSet& InOutActionSet, const int32 InActionId);
 	void RegisterSelectComponentsAction(FInteractiveToolActionSet& InOutActionSet, const int32 InActionId);
+	void RegisterSelectionFilterCyclingAction(FInteractiveToolActionSet& InOutActionSet, const int32 InActionId);
+	void RegisterSnapAction(FInteractiveToolActionSet& InOutActionSet, const int32 InActionId);
 
 	TArray<int32> GetSelectedComponents() const;
 
