Skeletal Mesh Tools: Expose polygroup manipulation tools.

diff --git a/Engine/Plugins/Experimental/Animation/SkeletalMeshModelingTools/Source/SkeletalMeshModelingTools/Private/SkeletalMeshModelingToolsEditorMode.cpp b/Engine/Plugins/Experimental/Animation/SkeletalMeshModelingTools/Source/SkeletalMeshModelingTools/Private/SkeletalMeshModelingToolsEditorMode.cpp
--- a/Engine/Plugins/Experimental/Animation/SkeletalMeshModelingTools/Source/SkeletalMeshModelingTools/Private/SkeletalMeshModelingToolsEditorMode.cpp
+++ b/Engine/Plugins/Experimental/Animation/SkeletalMeshModelingTools/Source/SkeletalMeshModelingTools/Private/SkeletalMeshModelingToolsEditorMode.cpp
@@ -25,6 +25,7 @@
 #include "ToolTargets/SkeletalMeshComponentToolTarget.h"
 #include "Components/SkeletalMeshComponent.h"
 
+#include "ConvertToPolygonsTool.h"
 #include "DeformMeshPolygonsTool.h"
 #include "DisplaceMeshTool.h"
 #include "DynamicMeshSculptTool.h"
@@ -34,6 +35,7 @@
 #include "ISkeletalMeshEditor.h"
 #include "LatticeDeformerTool.h"
 #include "MeshAttributePaintTool.h"
+#include "MeshGroupPaintTool.h"
 #include "MeshSpaceDeformerTool.h"
 #include "MeshVertexSculptTool.h"
 #include "ModelingToolsManagerActions.h"
@@ -208,7 +210,13 @@ void USkeletalMeshModelingToolsEditorMode::Enter()
 	RegisterTool(ToolManagerCommands.BeginRemoveOccludedTrianglesTool, TEXT("BeginRemoveOccludedTrianglesTool"), NewObject<URemoveOccludedTrianglesToolBuilder>());
 	RegisterTool(ToolManagerCommands.BeginProjectToTargetTool, TEXT("BeginProjectToTargetTool"), NewObject<UProjectToTargetToolBuilder>());
 	
-
+	RegisterTool(ToolManagerCommands.BeginPolyGroupsTool, TEXT("BeginPolyGroupsTool"), NewObject<UConvertToPolygonsToolBuilder>());
+	UMeshGroupPaintToolBuilder* MeshGroupPaintToolBuilder = NewObject<UMeshGroupPaintToolBuilder>();
+#if ENABLE_STYLUS_SUPPORT 
+	MeshGroupPaintToolBuilder->StylusAPI = StylusStateTracker.Get();
+#endif
+	RegisterTool(ToolManagerCommands.BeginMeshGroupPaintTool, TEXT("BeginMeshGroupPaintTool"), MeshGroupPaintToolBuilder);
+	
 	UMeshVertexSculptToolBuilder* MoveVerticesToolBuilder = NewObject<UMeshVertexSculptToolBuilder>();
 #if ENABLE_STYLUS_SUPPORT
 	MoveVerticesToolBuilder->StylusAPI = StylusStateTracker.Get();
diff --git a/Engine/Plugins/Experimental/Animation/SkeletalMeshModelingTools/Source/SkeletalMeshModelingTools/Private/SkeletalMeshModelingToolsEditorModeToolkit.cpp b/Engine/Plugins/Experimental/Animation/SkeletalMeshModelingTools/Source/SkeletalMeshModelingTools/Private/SkeletalMeshModelingToolsEditorModeToolkit.cpp
--- a/Engine/Plugins/Experimental/Animation/SkeletalMeshModelingTools/Source/SkeletalMeshModelingTools/Private/SkeletalMeshModelingToolsEditorModeToolkit.cpp
+++ b/Engine/Plugins/Experimental/Animation/SkeletalMeshModelingTools/Source/SkeletalMeshModelingTools/Private/SkeletalMeshModelingToolsEditorModeToolkit.cpp
@@ -174,7 +174,9 @@ void FSkeletalMeshModelingToolsEditorModeToolkit::RegisterPalettes()
 		Commands.BeginSkinWeightsBindingTool,
 		Commands.BeginSkinWeightsPaintTool,
 		Commands.BeginAttributeEditorTool,
-		Commands.BeginMeshAttributePaintTool
+		Commands.BeginMeshAttributePaintTool,
+		Commands.BeginPolyGroupsTool,
+		Commands.BeginMeshGroupPaintTool
 	});
 	ToolkitBuilder->AddPalette(
 		MakeShareable( new FToolPalette( Commands.LoadSkinTools.ToSharedRef(), SkinCommands ) ) );
