Skeleton Editing Tool: minor fixes
    - added tooltip for snapping button
    - enable create + snap with no selection

diff --git a/Engine/Plugins/Experimental/Animation/SkeletalMeshModelingTools/Source/SkeletalMeshModelingTools/Private/DetailCustomization/SkeletonEditingToolPropertyCustomizations.cpp b/Engine/Plugins/Experimental/Animation/SkeletalMeshModelingTools/Source/SkeletalMeshModelingTools/Private/DetailCustomization/SkeletonEditingToolPropertyCustomizations.cpp
--- a/Engine/Plugins/Experimental/Animation/SkeletalMeshModelingTools/Source/SkeletalMeshModelingTools/Private/DetailCustomization/SkeletonEditingToolPropertyCustomizations.cpp
+++ b/Engine/Plugins/Experimental/Animation/SkeletalMeshModelingTools/Source/SkeletalMeshModelingTools/Private/DetailCustomization/SkeletonEditingToolPropertyCustomizations.cpp
@@ -944,6 +944,12 @@ void FSkeletonEditingToolDetailCustomization::CustomizeComponentSelection(IDetai
 				const bool bCreate = Tool.IsValid() && Tool->GetOperation() == EEditingOperation::Create;
 				return bCreate ? LOCTEXT("SnapCreateButtonLabel", "Create") : LOCTEXT("SnapButtonLabel", "Snap");
 			})
+			.ToolTipText_Lambda([this]
+			{
+				const bool bCreate = Tool.IsValid() && Tool->GetOperation() == EEditingOperation::Create;
+				return bCreate ? LOCTEXT("SnapCreateButtonTooltip", "Create a new bone and snap it to the selected components. (V)") :
+								LOCTEXT("SnapButtonTooltip", "Snap the selected bone to the selected components. (V)");
+			})
 			.IsEnabled_Lambda([this]()
 			{
 				return Tool->HasSelectedComponent(); 
diff --git a/Engine/Plugins/Experimental/MeshModelingToolsetExp/Source/MeshModelingToolsEditorOnlyExp/Private/SkeletalMesh/SkeletonEditingTool.cpp b/Engine/Plugins/Experimental/MeshModelingToolsetExp/Source/MeshModelingToolsEditorOnlyExp/Private/SkeletalMesh/SkeletonEditingTool.cpp
--- a/Engine/Plugins/Experimental/MeshModelingToolsetExp/Source/MeshModelingToolsEditorOnlyExp/Private/SkeletalMesh/SkeletonEditingTool.cpp
+++ b/Engine/Plugins/Experimental/MeshModelingToolsetExp/Source/MeshModelingToolsEditorOnlyExp/Private/SkeletalMesh/SkeletonEditingTool.cpp
@@ -854,14 +854,18 @@ void USkeletonEditingTool::SnapBoneToComponentSelection(const bool bCreate)
 
 	const FReferenceSkeleton& RefSkeleton = Modifier->GetReferenceSkeleton();
 	const TArray<FName> Bones = GetSelectedBones();
-	const bool bHasValidBone = Bones.ContainsByPredicate([&](const FName& InBoneName)
-	{
-		return RefSkeleton.FindRawBoneIndex(InBoneName) > INDEX_NONE;
-	});
 
-	if (!bHasValidBone)
+	if (!bCreate)
 	{
-		return;
+		const bool bHasValidBone = Bones.ContainsByPredicate([&](const FName& InBoneName)
+	   {
+		   return RefSkeleton.FindRawBoneIndex(InBoneName) > INDEX_NONE;
+	   });
+
+		if (!bHasValidBone)
+		{
+			return;
+		}
 	}
 
 	const int32 BoneIndex = RefSkeleton.FindRawBoneIndex(GetCurrentBone());
