Skeletal Mesh Conversion: Allow selecting target path, prefixes, suffixes, and give a choice of root bone location.

diff --git a/Engine/Plugins/Experimental/Animation/SkeletalMeshModelingTools/Source/SkeletalMeshModelingTools/Private/SkeletalMeshModelingToolsMeshConverter.cpp b/Engine/Plugins/Experimental/Animation/SkeletalMeshModelingTools/Source/SkeletalMeshModelingTools/Private/SkeletalMeshModelingToolsMeshConverter.cpp
--- a/Engine/Plugins/Experimental/Animation/SkeletalMeshModelingTools/Source/SkeletalMeshModelingTools/Private/SkeletalMeshModelingToolsMeshConverter.cpp
+++ b/Engine/Plugins/Experimental/Animation/SkeletalMeshModelingTools/Source/SkeletalMeshModelingTools/Private/SkeletalMeshModelingToolsMeshConverter.cpp
@@ -10,6 +10,9 @@
 #include "StaticToSkeletalMeshConverter.h"
 #include "AssetRegistry/AssetRegistryModule.h"
 #include "Dialog/SCustomDialog.h"
+#include "Framework/Notifications/NotificationManager.h"
+#include "Widgets/Layout/SBox.h"
+#include "Widgets/Notifications/SNotificationList.h"
 
 DEFINE_LOG_CATEGORY_STATIC(LogSkeletalMeshModelingToolsMeshConverter, Log, All)
 
@@ -36,15 +39,34 @@ UObject* USkeletonFromStaticMeshFactory::FactoryCreateNew(
 	FFeedbackContext* InWarn
 	)
 {
-	static const FVector RootBoneRelativePosition(0.5, 0.5, 0.0);
-
 	USkeleton* Skeleton = NewObject<USkeleton>(InParent, InName, InFlags);
-	
-	if (!FStaticToSkeletalMeshConverter::InitializeSkeletonFromStaticMesh(Skeleton, StaticMesh, RootBoneRelativePosition))
+
+	FVector WantedRootPosition;
+	switch(PositionReference)
 	{
-		return nullptr;
+	case ERootBonePositionReference::Absolute:
+		WantedRootPosition = RootPosition;
+		break;
+		
+	case ERootBonePositionReference::Relative:
+		{
+			const FBox Bounds = StaticMesh->GetBoundingBox();
+			WantedRootPosition = Bounds.Min + (Bounds.Max - Bounds.Min) * RootPosition;		
+		}
+		break;
+		
+	default:
+		checkNoEntry()
+		break;
 	}
 
+	const TCHAR* RootBoneName = TEXT("Root"); 
+	FTransform RootTransform(FTransform::Identity);
+	RootTransform.SetTranslation(WantedRootPosition);
+
+	FReferenceSkeletonModifier Modifier(Skeleton);
+	Modifier.Add(FMeshBoneInfo(RootBoneName, RootBoneName, INDEX_NONE), RootTransform);
+
 	return Skeleton;
 }
 
@@ -148,18 +170,34 @@ void UStaticMeshToSkeletalMeshConvertOptions::PostEditChangeProperty(FPropertyCh
 	}
 }
 
+static void CreateAssetName(
+	const UPackage* InTemplatePackageName,
+	const FString& InTargetPackagePath,
+	const FString& InPrefixToRemove,
+	const FString& InPrefixToAdd,
+	const FString& InSuffixToAdd,
+	FString &OutNewAssetName,
+	FString& OutNewPackageName
+	)
+{
+	IAssetTools& AssetTools = FAssetToolsModule::GetModule().Get();
+	
+	FString TemplatePackageName = InTemplatePackageName->GetName();
+	FString TemplateAssetName = FPackageName::GetLongPackageAssetName(TemplatePackageName);
+	TemplateAssetName.RemoveFromStart(InPrefixToRemove);
+	TemplateAssetName.InsertAt(0, InPrefixToAdd);
+
+	TemplatePackageName = InTargetPackagePath + TEXT("/") + TemplateAssetName; 
+	
+	AssetTools.CreateUniqueAssetName(TemplatePackageName, InSuffixToAdd, OutNewPackageName, OutNewAssetName);
+}
+
 static bool ConvertSingleMeshToSkeletalMesh(
 	const UStaticMeshToSkeletalMeshConvertOptions* InOptions,
 	UStaticMesh* InStaticMesh,
 	TArray<UObject*>& OutObjectsAdded
 	)
 {
-	
-	// Make these configurable.
-	// ReSharper disable CppTooWideScope
-	static const TCHAR* DefaultSkeletonSuffix = TEXT("_Skel");
-	static const TCHAR* DefaultSkeletalMeshSuffix = TEXT("_SkelMesh");
-
 	IAssetTools& AssetTools = FAssetToolsModule::GetModule().Get();
 	FReferenceSkeleton ReferenceSkeleton;
 	USkeleton* Skeleton = nullptr;
@@ -171,11 +209,28 @@ static bool ConvertSingleMeshToSkeletalMesh(
 		{
 			FString SkeletonName;
 			FString SkeletonPackageName;
-			AssetTools.CreateUniqueAssetName(InStaticMesh->GetOutermost()->GetName(), DefaultSkeletonSuffix, SkeletonPackageName, SkeletonName);
+
+			CreateAssetName(InStaticMesh->GetPackage(), InOptions->DestinationPath.Path, InOptions->PrefixToRemove, InOptions->SkeletonPrefixToAdd, InOptions->SkeletonSuffixToAdd, SkeletonName, SkeletonPackageName);
 
 			USkeletonFromStaticMeshFactory* SkeletonFactory = NewObject<USkeletonFromStaticMeshFactory>();
 			SkeletonFactory->StaticMesh = InStaticMesh;
 
+			switch(InOptions->RootBonePlacement)
+			{
+			case ERootBonePlacementOptions::BottomCenter:
+				SkeletonFactory->RootPosition = FVector(0.5, 0.5, 0.0);
+				SkeletonFactory->PositionReference = ERootBonePositionReference::Relative;
+				break;
+			case ERootBonePlacementOptions::Center:
+				SkeletonFactory->RootPosition = FVector(0.5, 0.5, 0.5);
+				SkeletonFactory->PositionReference = ERootBonePositionReference::Relative;
+				break;
+			case ERootBonePlacementOptions::Origin:
+				SkeletonFactory->RootPosition = FVector(0.0, 0.0, 0.0);
+				SkeletonFactory->PositionReference = ERootBonePositionReference::Absolute;
+				break;
+			}
+
 			Skeleton = Cast<USkeleton>(AssetTools.CreateAsset(SkeletonName, FPackageName::GetLongPackagePath(SkeletonPackageName), USkeleton::StaticClass(), SkeletonFactory));
 			ReferenceSkeleton = Skeleton->GetReferenceSkeleton();
 			
@@ -215,7 +270,8 @@ static bool ConvertSingleMeshToSkeletalMesh(
 	{
 		FString SkeletalMeshName;
 		FString SkeletalMeshPackageName;
-		AssetTools.CreateUniqueAssetName(InStaticMesh->GetOutermost()->GetName(), DefaultSkeletalMeshSuffix, SkeletalMeshPackageName, SkeletalMeshName);
+		
+		CreateAssetName(InStaticMesh->GetPackage(), InOptions->DestinationPath.Path, InOptions->PrefixToRemove, InOptions->SkeletalMeshPrefixToAdd, InOptions->SkeletalMeshSuffixToAdd, SkeletalMeshName, SkeletalMeshPackageName);
 
 		USkeletalMeshFromStaticMeshFactory* SkeletalMeshFactory = NewObject<USkeletalMeshFromStaticMeshFactory>();
 		SkeletalMeshFactory->StaticMesh = InStaticMesh;
@@ -267,11 +323,14 @@ void ConvertStaticMeshAssetsToSkeletalMeshesInteractive(
 
 		MeshesToConvert.Add(StaticMesh);
 	}
-	
+	if (MeshesToConvert.IsEmpty())
+	{
+		return;
+	}
 
 	UStaticMeshToSkeletalMeshConvertOptions* Options = NewObject<UStaticMeshToSkeletalMeshConvertOptions>();
 	FPropertyEditorModule& PropertyEditorModule = FModuleManager::GetModuleChecked<FPropertyEditorModule>("PropertyEditor");
-
+	
 	// If the object paths don't resolve properly, the details view will still show something, which will be invalidated
 	// later.
 	if (Cast<USkeleton>(Options->Skeleton.ResolveObject()) == nullptr)
@@ -283,9 +342,24 @@ void ConvertStaticMeshAssetsToSkeletalMeshesInteractive(
 		Options->SkeletalMesh.Reset();
 	}
 
+	Options->DestinationPath.Path = FPackageName::GetLongPackagePath(MeshesToConvert[0]->GetPackage()->GetPathName());
+	
 	auto OnConvertLambda = [Options, MeshesToConvert]()
 	{
 		Options->SaveConfig();
+
+		FPackageName::EErrorCode PathErrorCode;
+		if (!FPackageName::IsValidLongPackageName(Options->DestinationPath.Path, false, &PathErrorCode))
+		{
+			const FText Message = FPackageName::FormatErrorAsText(Options->DestinationPath.Path, PathErrorCode);
+			FNotificationInfo WarningNotification(Message);
+			WarningNotification.ExpireDuration = 5.0f;
+			WarningNotification.bFireAndForget = true;
+			WarningNotification.bUseLargeFont = false;
+			WarningNotification.Image = FCoreStyle::Get().GetBrush(TEXT("MessageLog.Error"));
+			FSlateNotificationManager::Get().AddNotification(WarningNotification);
+			UE_LOG(LogSkeletalMeshModelingToolsMeshConverter, Error, TEXT("%s"), *Message.ToString());
+		}
 	
 		TArray<UObject*> ObjectsAdded;
 		if (ConvertMultipleMeshesToSkeletalMesh(Options, MeshesToConvert, ObjectsAdded))
@@ -312,7 +386,11 @@ void ConvertStaticMeshAssetsToSkeletalMeshesInteractive(
 		.Title(LOCTEXT("OptionsDialogTitle", "Static Mesh Conversion Options"))
 		.Content()
 		[
-			DetailsView
+			SNew(SBox)
+			.MinDesiredWidth(450)
+			[
+				DetailsView
+			]
 		]
 	.Buttons({
 		SCustomDialog::FButton(LOCTEXT("DialogButtonConvert", "Convert"), FSimpleDelegate::CreateLambda(OnConvertLambda)),
diff --git a/Engine/Plugins/Experimental/Animation/SkeletalMeshModelingTools/Source/SkeletalMeshModelingTools/Private/SkeletalMeshModelingToolsMeshConverter.h b/Engine/Plugins/Experimental/Animation/SkeletalMeshModelingTools/Source/SkeletalMeshModelingTools/Private/SkeletalMeshModelingToolsMeshConverter.h
--- a/Engine/Plugins/Experimental/Animation/SkeletalMeshModelingTools/Source/SkeletalMeshModelingTools/Private/SkeletalMeshModelingToolsMeshConverter.h
+++ b/Engine/Plugins/Experimental/Animation/SkeletalMeshModelingTools/Source/SkeletalMeshModelingTools/Private/SkeletalMeshModelingToolsMeshConverter.h
@@ -14,6 +14,13 @@ class USkeleton;
 class UStaticMesh;
 struct FAssetData;
 
+UENUM()
+enum class ERootBonePositionReference
+{
+	Relative,
+	Absolute
+};
+
 
 UCLASS(Hidden, MinimalAPI)
 class USkeletonFromStaticMeshFactory :
@@ -37,6 +44,12 @@ public:
 	
 	UPROPERTY()
 	TObjectPtr<UStaticMesh> StaticMesh;
+
+	UPROPERTY()
+	FVector RootPosition = {0.5, 0.5, 0.0};
+
+	UPROPERTY()
+	ERootBonePositionReference PositionReference = ERootBonePositionReference::Relative;
 };
 
 
@@ -76,12 +89,21 @@ public:
 UENUM()
 enum class EReferenceSkeletonImportOption
 {
-	CreateNew,
-	UseExistingSkeleton,
-	UseExistingSkeletalMesh
+	CreateNew UMETA(ToolTip="Create a new skeleton with a single root bone"),
+	UseExistingSkeleton UMETA(ToolTip="Use an already created skeleton asset"),
+	UseExistingSkeletalMesh UMETA(ToolTip="Use a skeleton use by another skeletal mesh asset")
+};
+
+UENUM()
+enum class ERootBonePlacementOptions
+{
+	BottomCenter UMETA(DisplayName="Bottom center", ToolTip="Bottom-center of the static mesh's bounding box"),
+	Center UMETA(DisplayName="Center", ToolTip="Center of the static mesh's bounding box"),
+	Origin UMETA(DisplayName="Origin", ToolTip="At the static mesh's origin (0,0,0)"),
 };
 
 
+
 UCLASS(config=EditorPerProjectUserSettings, HideCategories=Object, MinimalAPI)
 class UStaticMeshToSkeletalMeshConvertOptions :
 	public UObject,
@@ -89,19 +111,40 @@ class UStaticMeshToSkeletalMeshConvertOptions :
 {
 	GENERATED_BODY()
 public:
+	UPROPERTY(EditAnywhere, Config, Category="Location", meta=(ContentDir))
+	FDirectoryPath DestinationPath;
+	
 
-	UPROPERTY(EditAnywhere, Config, Category="Import")
+	UPROPERTY(EditAnywhere, Config, Category="Skeleton", meta=(DisplayName="Creation Type"))
 	EReferenceSkeletonImportOption SkeletonImportOption = EReferenceSkeletonImportOption::CreateNew;
 
-	UPROPERTY(EditAnywhere, Config, Category="Import", Meta=(AllowedClasses="/Script/Engine.Skeleton", EditCondition="SkeletonImportOption==EReferenceSkeletonImportOption::UseExistingSkeleton", EditConditionHides))
+	UPROPERTY(EditAnywhere, Config, Category="Skeleton", meta=(EditCondition="SkeletonImportOption==EReferenceSkeletonImportOption::CreateNew", EditConditionHides))
+	ERootBonePlacementOptions RootBonePlacement = ERootBonePlacementOptions::BottomCenter;
+	
+	UPROPERTY(EditAnywhere, Config, Category="Skeleton", meta=(AllowedClasses="/Script/Engine.Skeleton", EditCondition="SkeletonImportOption==EReferenceSkeletonImportOption::UseExistingSkeleton", EditConditionHides))
 	FSoftObjectPath Skeleton;
 
-	UPROPERTY(EditAnywhere, Config, Category="Import", Meta=(AllowedClasses="/Script/Engine.SkeletalMesh", EditCondition="SkeletonImportOption==EReferenceSkeletonImportOption::UseExistingSkeletalMesh", EditConditionHides))
+	UPROPERTY(EditAnywhere, Config, Category="Skeleton", meta=(AllowedClasses="/Script/Engine.SkeletalMesh", EditCondition="SkeletonImportOption==EReferenceSkeletonImportOption::UseExistingSkeletalMesh", EditConditionHides))
 	FSoftObjectPath SkeletalMesh;
 
-	UPROPERTY(EditAnywhere, Config, Category="Import", meta=(EditCondition="SkeletonImportOption!=EReferenceSkeletonImportOption::CreateNew", EditConditionHides))
+	UPROPERTY(EditAnywhere, Config, Category="Skeleton", meta=(EditCondition="SkeletonImportOption!=EReferenceSkeletonImportOption::CreateNew", EditConditionHides))
 	FBoneReference BindingBoneName;
 
+	UPROPERTY(EditAnywhere, Config, Category="Naming")
+	FString PrefixToRemove = TEXT("SM_");
+	
+	UPROPERTY(EditAnywhere, Config, Category="Naming", meta=(DisplayName="Skeletal Mesh Prefix"))
+	FString SkeletalMeshPrefixToAdd = TEXT("SKM_");
+
+	UPROPERTY(EditAnywhere, Config, Category="Naming", meta=(DisplayName="Skeletal Mesh Prefix"))
+	FString SkeletalMeshSuffixToAdd;
+	
+	UPROPERTY(EditAnywhere, Config, Category="Naming", meta=(DisplayName="Skeleton Prefix"))
+	FString SkeletonPrefixToAdd = TEXT("SK_");
+
+	UPROPERTY(EditAnywhere, Config, Category="Naming", meta=(DisplayName="Skeleton Prefix"))
+	FString SkeletonSuffixToAdd;
+	
 	// IBoneReferenceSkeletonProvider implementation.
 	virtual USkeleton* GetSkeleton(bool& bInvalidSkeletonIsError, const IPropertyHandle* PropertyHandle) override;
 	
diff --git a/Engine/Source/Editor/DetailCustomizations/Private/DirectoryPathStructCustomization.cpp b/Engine/Source/Editor/DetailCustomizations/Private/DirectoryPathStructCustomization.cpp
--- a/Engine/Source/Editor/DetailCustomizations/Private/DirectoryPathStructCustomization.cpp
+++ b/Engine/Source/Editor/DetailCustomizations/Private/DirectoryPathStructCustomization.cpp
@@ -131,6 +131,7 @@ FReply FDirectoryPathStructCustomization::OnPickContent(TSharedRef<IPropertyHand
 {
 	FContentBrowserModule& ContentBrowserModule = FModuleManager::LoadModuleChecked<FContentBrowserModule>("ContentBrowser");
 	FPathPickerConfig PathPickerConfig;
+	PropertyHandle->GetValue(PathPickerConfig.DefaultPath);
 	PathPickerConfig.bAllowContextMenu = false;
 	PathPickerConfig.OnPathSelected = FOnPathSelected::CreateSP(this, &FDirectoryPathStructCustomization::OnPathPicked, PropertyHandle);
 
diff --git a/Engine/Source/Editor/Persona/Private/AnimGraphNodeDetails.cpp b/Engine/Source/Editor/Persona/Private/AnimGraphNodeDetails.cpp
--- a/Engine/Source/Editor/Persona/Private/AnimGraphNodeDetails.cpp
+++ b/Engine/Source/Editor/Persona/Private/AnimGraphNodeDetails.cpp
@@ -538,6 +538,7 @@ void FBoneReferenceCustomization::CustomizeHeader( TSharedRef<IPropertyHandle> S
 		.ValueContent()
 		[
 			SNew(SEditableTextBox)
+			.Font(StructCustomizationUtils.GetRegularFont())
 			.ToolTipText(StructPropertyHandle->GetToolTipText())
 			.OnTextCommitted_Lambda([this](const FText& InText, ETextCommit::Type InTextCommit)
 			{
