Skeletal Mesh Converter: Fix labels and disallow invalid names

#rnx

diff --git a/Engine/Plugins/Experimental/Animation/SkeletalMeshModelingTools/Source/SkeletalMeshModelingTools/Private/SkeletalMeshModelingToolsMeshConverter.cpp b/Engine/Plugins/Experimental/Animation/SkeletalMeshModelingTools/Source/SkeletalMeshModelingTools/Private/SkeletalMeshModelingToolsMeshConverter.cpp
--- a/Engine/Plugins/Experimental/Animation/SkeletalMeshModelingTools/Source/SkeletalMeshModelingTools/Private/SkeletalMeshModelingToolsMeshConverter.cpp
+++ b/Engine/Plugins/Experimental/Animation/SkeletalMeshModelingTools/Source/SkeletalMeshModelingTools/Private/SkeletalMeshModelingToolsMeshConverter.cpp
@@ -18,6 +18,31 @@ DEFINE_LOG_CATEGORY_STATIC(LogSkeletalMeshModelingToolsMeshConverter, Log, All)
 
 #define LOCTEXT_NAMESPACE "SkeletalMeshModelingToolsMeshConverter"
 
+static void ShowEditorMessage(ELogVerbosity::Type InMessageType, const FText& InMessage, const FText* InLogMessage)
+{
+	FNotificationInfo Notification(InMessage);
+	Notification.bUseSuccessFailIcons = true;
+	Notification.FadeOutDuration = 5.0f;
+
+	SNotificationItem::ECompletionState State = SNotificationItem::CS_Success;
+
+	switch(InMessageType)
+	{
+	case ELogVerbosity::Warning:
+		UE_LOG(LogSkeletalMeshModelingToolsMeshConverter, Warning, TEXT("%s"), InLogMessage ? *InLogMessage->ToString() : *InMessage.ToString());
+		break;
+	case ELogVerbosity::Error:
+		State = SNotificationItem::CS_Fail;
+		UE_LOG(LogSkeletalMeshModelingToolsMeshConverter, Error, TEXT("%s"), InLogMessage ? *InLogMessage->ToString() : *InMessage.ToString());
+		break;
+	default:
+		checkNoEntry();
+	}
+	
+	FSlateNotificationManager::Get().AddNotification(Notification)->SetCompletionState(State);
+}
+
+
 
 USkeletonFromStaticMeshFactory::USkeletonFromStaticMeshFactory(
 	const FObjectInitializer& InObjectInitializer
@@ -230,6 +255,10 @@ static bool ConvertSingleMeshToSkeletalMesh(
 			}
 
 			Skeleton = Cast<USkeleton>(AssetTools.CreateAsset(SkeletonName, FPackageName::GetLongPackagePath(SkeletonPackageName), USkeleton::StaticClass(), SkeletonFactory));
+			if (!Skeleton)
+			{
+				return false;
+			}
 			ReferenceSkeleton = Skeleton->GetReferenceSkeleton();
 			
 			OutObjectsAdded.Add(Skeleton);
@@ -305,6 +334,60 @@ static bool ConvertMultipleMeshesToSkeletalMesh(
 	return true;
 }
 
+static void DoConversion(
+	UStaticMeshToSkeletalMeshConvertOptions* InOptions,
+	const TArray<UStaticMesh*>& InMeshesToConvert
+	)
+{
+	auto IsValidPathPart = [](const FString& InPathPart, FStringView InPathPartName, const FText& InMessage)-> bool
+	{
+		FText FailureReason;
+		const FText FailureContext = FText::FromStringView(InPathPartName);
+		if (!FName::IsValidXName(InPathPart, INVALID_OBJECTNAME_CHARACTERS INVALID_LONGPACKAGE_CHARACTERS, &FailureReason, &FailureContext))
+		{
+			ShowEditorMessage(ELogVerbosity::Error, InMessage, &FailureReason); 
+			return false;
+		}
+		return true;
+	};
+	
+	if (!IsValidPathPart(InOptions->DestinationPath.Path, TEXT("Destination Path"), LOCTEXT("InvalidDestinationPath", "Invalid Characters in Destination Path")))
+	{
+		return;
+	}
+
+	if (!IsValidPathPart(InOptions->SkeletalMeshPrefixToAdd, TEXT("Skeletal Mesh Prefix"), LOCTEXT("InvalidSkeletalMeshPrefix", "Invalid Characters in Skeletal Mesh Prefix")) ||
+		!IsValidPathPart(InOptions->SkeletalMeshSuffixToAdd, TEXT("Skeletal Mesh Suffix"), LOCTEXT("InvalidSkeletalMeshPrefix", "Invalid Characters in Skeletal Mesh Suffix")))
+	{	
+		return;
+	}
+				
+	if (InOptions->SkeletonImportOption == EReferenceSkeletonImportOption::CreateNew)
+	{
+		if (!IsValidPathPart(InOptions->SkeletonPrefixToAdd, TEXT("Skeleton Prefix"), LOCTEXT("InvalidSkeletonPrefix", "Invalid Characters in Skeleton Prefix")) ||
+			!IsValidPathPart(InOptions->SkeletonSuffixToAdd, TEXT("Skeleton Suffix"), LOCTEXT("InvalidSkeletonSuffix", "Invalid Characters in Skeleton Suffix")))
+		{
+			return;
+		}
+	}
+
+	InOptions->SaveConfig();
+	
+	TArray<UObject*> ObjectsAdded;
+	if (ConvertMultipleMeshesToSkeletalMesh(InOptions, InMeshesToConvert, ObjectsAdded))
+	{
+		FAssetToolsModule::GetModule().Get().SyncBrowserToAssets(ObjectsAdded);
+	}
+	else
+	{
+		for (UObject* ObjectToDelete: ObjectsAdded)
+		{
+			FAssetRegistryModule::AssetDeleted(ObjectToDelete);
+			ObjectToDelete->Rename(nullptr, GetTransientPackage(), REN_DontCreateRedirectors | REN_ForceNoResetLoaders | REN_NonTransactional);
+		}
+	}
+}
+
 void ConvertStaticMeshAssetsToSkeletalMeshesInteractive(
 	const TArray<FAssetData>& InStaticMeshAssets
 	)
@@ -342,38 +425,6 @@ void ConvertStaticMeshAssetsToSkeletalMeshesInteractive(
 
 	Options->DestinationPath.Path = FPackageName::GetLongPackagePath(MeshesToConvert[0]->GetPackage()->GetPathName());
 	
-	auto OnConvertLambda = [Options, MeshesToConvert]()
-	{
-		Options->SaveConfig();
-
-		FPackageName::EErrorCode PathErrorCode;
-		if (!FPackageName::IsValidLongPackageName(Options->DestinationPath.Path, false, &PathErrorCode))
-		{
-			const FText Message = FPackageName::FormatErrorAsText(Options->DestinationPath.Path, PathErrorCode);
-			FNotificationInfo WarningNotification(Message);
-			WarningNotification.ExpireDuration = 5.0f;
-			WarningNotification.bFireAndForget = true;
-			WarningNotification.bUseLargeFont = false;
-			WarningNotification.Image = FCoreStyle::Get().GetBrush(TEXT("MessageLog.Error"));
-			FSlateNotificationManager::Get().AddNotification(WarningNotification);
-			UE_LOG(LogSkeletalMeshModelingToolsMeshConverter, Error, TEXT("%s"), *Message.ToString());
-		}
-	
-		TArray<UObject*> ObjectsAdded;
-		if (ConvertMultipleMeshesToSkeletalMesh(Options, MeshesToConvert, ObjectsAdded))
-		{
-			FAssetToolsModule::GetModule().Get().SyncBrowserToAssets(ObjectsAdded);
-		}
-		else
-		{
-			for (UObject* ObjectToDelete: ObjectsAdded)
-			{
-				FAssetRegistryModule::AssetDeleted(ObjectToDelete);
-				ObjectToDelete->Rename(nullptr, GetTransientPackage(), REN_DontCreateRedirectors | REN_ForceNoResetLoaders | REN_NonTransactional);
-			}
-		}
-	};
-
 	FDetailsViewArgs DetailsViewArgs;
 	DetailsViewArgs.bAllowSearch = false;
 	DetailsViewArgs.NameAreaSettings = FDetailsViewArgs::HideNameArea;
@@ -391,7 +442,10 @@ void ConvertStaticMeshAssetsToSkeletalMeshesInteractive(
 			]
 		]
 	.Buttons({
-		SCustomDialog::FButton(LOCTEXT("DialogButtonConvert", "Convert"), FSimpleDelegate::CreateLambda(OnConvertLambda)),
+		SCustomDialog::FButton(LOCTEXT("DialogButtonConvert", "Convert"), FSimpleDelegate::CreateLambda([Options, MeshesToConvert]()
+		{
+			DoConversion(Options, MeshesToConvert);
+		})),
 		SCustomDialog::FButton(LOCTEXT("DialogButtonCancel", "Cancel"))
 	});
 
diff --git a/Engine/Plugins/Experimental/Animation/SkeletalMeshModelingTools/Source/SkeletalMeshModelingTools/Private/SkeletalMeshModelingToolsMeshConverter.h b/Engine/Plugins/Experimental/Animation/SkeletalMeshModelingTools/Source/SkeletalMeshModelingTools/Private/SkeletalMeshModelingToolsMeshConverter.h
--- a/Engine/Plugins/Experimental/Animation/SkeletalMeshModelingTools/Source/SkeletalMeshModelingTools/Private/SkeletalMeshModelingToolsMeshConverter.h
+++ b/Engine/Plugins/Experimental/Animation/SkeletalMeshModelingTools/Source/SkeletalMeshModelingTools/Private/SkeletalMeshModelingToolsMeshConverter.h
@@ -137,13 +137,13 @@ public:
 	UPROPERTY(EditAnywhere, Config, Category="Naming", meta=(DisplayName="Skeletal Mesh Prefix"))
 	FString SkeletalMeshPrefixToAdd = TEXT("SKM_");
 
-	UPROPERTY(EditAnywhere, Config, Category="Naming", meta=(DisplayName="Skeletal Mesh Prefix"))
+	UPROPERTY(EditAnywhere, Config, Category="Naming", meta=(DisplayName="Skeletal Mesh Suffix"))
 	FString SkeletalMeshSuffixToAdd;
 	
-	UPROPERTY(EditAnywhere, Config, Category="Naming", meta=(DisplayName="Skeleton Prefix"))
+	UPROPERTY(EditAnywhere, Config, Category="Naming", meta=(DisplayName="Skeleton Prefix", EditCondition="SkeletonImportOption==EReferenceSkeletonImportOption::CreateNew"))
 	FString SkeletonPrefixToAdd = TEXT("SK_");
 
-	UPROPERTY(EditAnywhere, Config, Category="Naming", meta=(DisplayName="Skeleton Prefix"))
+	UPROPERTY(EditAnywhere, Config, Category="Naming", meta=(DisplayName="Skeleton Suffix", EditCondition="SkeletonImportOption==EReferenceSkeletonImportOption::CreateNew"))
 	FString SkeletonSuffixToAdd;
 	
 	// IBoneReferenceSkeletonProvider implementation.
