Skeletal Mesh Editing Tools: Ensure plugin dependencies are marked as editor-only.

#rnx

diff --git a/Engine/Plugins/Experimental/Animation/SkeletalMeshModelingTools/SkeletalMeshModelingTools.uplugin b/Engine/Plugins/Experimental/Animation/SkeletalMeshModelingTools/SkeletalMeshModelingTools.uplugin
--- a/Engine/Plugins/Experimental/Animation/SkeletalMeshModelingTools/SkeletalMeshModelingTools.uplugin
+++ b/Engine/Plugins/Experimental/Animation/SkeletalMeshModelingTools/SkeletalMeshModelingTools.uplugin
@@ -21,13 +21,10 @@
 		}
 	],
 	"Plugins": [
-		{
-			"Name": "MeshModelingToolsetExp",
-			"Enabled": true
-		},
 		{
 			"Name": "ModelingToolsEditorMode",
-			"Enabled": true
+			"Enabled": true,
+			"TargetAllowList": [ "Editor" ]
 		}
 	]
 }
