Skeleton Editing Tool: disable snapping if no bone selected

#rnx

diff --git a/Engine/Plugins/Experimental/Animation/SkeletalMeshModelingTools/Source/SkeletalMeshModelingTools/Private/DetailCustomization/SkeletonEditingToolPropertyCustomizations.cpp b/Engine/Plugins/Experimental/Animation/SkeletalMeshModelingTools/Source/SkeletalMeshModelingTools/Private/DetailCustomization/SkeletonEditingToolPropertyCustomizations.cpp
--- a/Engine/Plugins/Experimental/Animation/SkeletalMeshModelingTools/Source/SkeletalMeshModelingTools/Private/DetailCustomization/SkeletonEditingToolPropertyCustomizations.cpp
+++ b/Engine/Plugins/Experimental/Animation/SkeletalMeshModelingTools/Source/SkeletalMeshModelingTools/Private/DetailCustomization/SkeletonEditingToolPropertyCustomizations.cpp
@@ -952,7 +952,12 @@ void FSkeletonEditingToolDetailCustomization::CustomizeComponentSelection(IDetai
 			})
 			.IsEnabled_Lambda([this]()
 			{
-				return Tool->HasSelectedComponent(); 
+				if (!Tool.IsValid())
+				{
+					return false;
+				}
+				const bool bCreate = Tool->GetOperation() == EEditingOperation::Create;
+				return bCreate ? Tool->HasSelectedComponent() : Tool->HasSelectedComponent() && !Tool->GetSelection().IsEmpty(); 
 			})
 			.OnClicked_Lambda([this]()
 			{
