Skeletal Mesh Converter: Fix copy-paste error with LOCTEXT

#rnx

diff --git a/Engine/Plugins/Experimental/Animation/SkeletalMeshModelingTools/Source/SkeletalMeshModelingTools/Private/SkeletalMeshModelingToolsMeshConverter.cpp b/Engine/Plugins/Experimental/Animation/SkeletalMeshModelingTools/Source/SkeletalMeshModelingTools/Private/SkeletalMeshModelingToolsMeshConverter.cpp
--- a/Engine/Plugins/Experimental/Animation/SkeletalMeshModelingTools/Source/SkeletalMeshModelingTools/Private/SkeletalMeshModelingToolsMeshConverter.cpp
+++ b/Engine/Plugins/Experimental/Animation/SkeletalMeshModelingTools/Source/SkeletalMeshModelingTools/Private/SkeletalMeshModelingToolsMeshConverter.cpp
@@ -357,7 +357,7 @@ static void DoConversion(
 	}
 
 	if (!IsValidPathPart(InOptions->SkeletalMeshPrefixToAdd, TEXT("Skeletal Mesh Prefix"), LOCTEXT("InvalidSkeletalMeshPrefix", "Invalid Characters in Skeletal Mesh Prefix")) ||
-		!IsValidPathPart(InOptions->SkeletalMeshSuffixToAdd, TEXT("Skeletal Mesh Suffix"), LOCTEXT("InvalidSkeletalMeshPrefix", "Invalid Characters in Skeletal Mesh Suffix")))
+		!IsValidPathPart(InOptions->SkeletalMeshSuffixToAdd, TEXT("Skeletal Mesh Suffix"), LOCTEXT("InvalidSkeletalMeshSuffix", "Invalid Characters in Skeletal Mesh Suffix")))
 	{	
 		return;
 	}
