Skeletal Mesh Converter: Reduce overzealous checking of long package path characters.

#rnx

#changelist validated
#virtualized

diff --git a/Engine/Plugins/Experimental/Animation/SkeletalMeshModelingTools/Source/SkeletalMeshModelingTools/Private/SkeletalMeshModelingToolsMeshConverter.cpp b/Engine/Plugins/Experimental/Animation/SkeletalMeshModelingTools/Source/SkeletalMeshModelingTools/Private/SkeletalMeshModelingToolsMeshConverter.cpp
--- a/Engine/Plugins/Experimental/Animation/SkeletalMeshModelingTools/Source/SkeletalMeshModelingTools/Private/SkeletalMeshModelingToolsMeshConverter.cpp
+++ b/Engine/Plugins/Experimental/Animation/SkeletalMeshModelingTools/Source/SkeletalMeshModelingTools/Private/SkeletalMeshModelingToolsMeshConverter.cpp
@@ -350,9 +350,12 @@ static void DoConversion(
 		}
 		return true;
 	};
-	
-	if (!IsValidPathPart(InOptions->DestinationPath.Path, TEXT("Destination Path"), LOCTEXT("InvalidDestinationPath", "Invalid Characters in Destination Path")))
+
+	FPackageName::EErrorCode PackagePathErrorCode;
+	if (!FPackageName::IsValidLongPackageName(InOptions->DestinationPath.Path, false, &PackagePathErrorCode))
 	{
+		const FText LogMessage = FPackageName::FormatErrorAsText(InOptions->DestinationPath.Path, PackagePathErrorCode);
+		ShowEditorMessage(ELogVerbosity::Error, LOCTEXT("InvalidDestinationPath", "Destination Path is Invalid"), &LogMessage); 
 		return;
 	}
 
