Skeletal Mesh Modeling Plugin: Store the activation state

#rnx

diff --git a/Engine/Plugins/Experimental/Animation/SkeletalMeshModelingTools/Source/SkeletalMeshModelingTools/Private/SkeletalMeshModelingToolsModule.cpp b/Engine/Plugins/Experimental/Animation/SkeletalMeshModelingTools/Source/SkeletalMeshModelingTools/Private/SkeletalMeshModelingToolsModule.cpp
--- a/Engine/Plugins/Experimental/Animation/SkeletalMeshModelingTools/Source/SkeletalMeshModelingTools/Private/SkeletalMeshModelingToolsModule.cpp
+++ b/Engine/Plugins/Experimental/Animation/SkeletalMeshModelingTools/Source/SkeletalMeshModelingTools/Private/SkeletalMeshModelingToolsModule.cpp
@@ -19,6 +19,7 @@
 #include "SkeletalMeshToolMenuContext.h"
 #include "ToolMenus.h"
 #include "DetailCustomization/SkeletonEditingToolPropertyCustomizations.h"
+#include "Misc/ConfigCacheIni.h"
 #include "SkeletalMesh/SkeletonEditingTool.h"
 #include "Styling/SlateIconFinder.h"
 #include "WorkflowOrientedApp/ApplicationMode.h"
@@ -30,6 +31,9 @@ DEFINE_LOG_CATEGORY(LogSkeletalMeshModelingTools);
 
 IMPLEMENT_MODULE(FSkeletalMeshModelingToolsModule, SkeletalMeshModelingTools);
 
+static const TCHAR* ConfigSection = TEXT("SkeletalMeshModelingTools");
+static const TCHAR* ConfigEditingModeActiveKey = TEXT("EditingModeActive");
+
 
 
 void FSkeletalMeshModelingToolsModule::StartupModule()
@@ -47,7 +51,7 @@ void FSkeletalMeshModelingToolsModule::StartupModule()
 	SkelMeshEditorExtenderHandle = ToolbarExtenders.Last().GetHandle();
 	// register post-init callback with skeletal mesh editor
 	TArray<ISkeletalMeshEditorModule::FOnSkeletalMeshEditorInitialized>& PostInitDelegates = SkelMeshEditorModule.GetPostEditorInitDelegates();
-	PostInitDelegates.Add(ISkeletalMeshEditorModule::FOnSkeletalMeshEditorInitialized::CreateRaw(this, &FSkeletalMeshModelingToolsModule::OnToggleEditingToolsMode));
+	PostInitDelegates.Add(ISkeletalMeshEditorModule::FOnSkeletalMeshEditorInitialized::CreateRaw(this, &FSkeletalMeshModelingToolsModule::CheckEnableEditingToolModeOnOpen));
 	SkelMeshEditorPostInitHandle = PostInitDelegates.Last().GetHandle();
 
 	FCoreDelegates::OnPostEngineInit.AddRaw(this, &FSkeletalMeshModelingToolsModule::OnPostEngineInit);
@@ -170,6 +174,9 @@ void FSkeletalMeshModelingToolsModule::OnToggleEditingToolsMode(TWeakPtr<ISkelet
 			EditorModeManager.DeactivateMode(USkeletalMeshModelingToolsEditorMode::Id);
 		}
 
+		// Update the stored state of the editing tools active state.
+		GConfig->SetBool(ConfigSection, ConfigEditingModeActiveKey, IsEditingToolModeActive(InSkeletalMeshEditor), GEditorPerProjectIni);
+
 		// make sure SkeletonSelection is active when toggling the mode, as they are compatible.
 		// it will be deactivated later when entering a tool 
 		if (!EditorModeManager.IsModeActive(FPersonaEditModes::SkeletonSelection))
@@ -179,6 +186,16 @@ void FSkeletalMeshModelingToolsModule::OnToggleEditingToolsMode(TWeakPtr<ISkelet
 	}
 }
 
+void FSkeletalMeshModelingToolsModule::CheckEnableEditingToolModeOnOpen(TWeakPtr<ISkeletalMeshEditor> InSkeletalMeshEditor)
+{
+	bool bEditingModeActive = false;
+	GConfig->GetBool(ConfigSection, ConfigEditingModeActiveKey, bEditingModeActive, GEditorPerProjectIni);
+	if (bEditingModeActive && !IsEditingToolModeActive(InSkeletalMeshEditor))
+	{
+		OnToggleEditingToolsMode(InSkeletalMeshEditor);
+	}
+}
+
 void FSkeletalMeshModelingToolsModule::OnPostEngineInit()
 {
 	FSkeletalMeshModelingToolsActionCommands::RegisterAllToolActions();
diff --git a/Engine/Plugins/Experimental/Animation/SkeletalMeshModelingTools/Source/SkeletalMeshModelingTools/Private/SkeletalMeshModelingToolsModule.h b/Engine/Plugins/Experimental/Animation/SkeletalMeshModelingTools/Source/SkeletalMeshModelingTools/Private/SkeletalMeshModelingToolsModule.h
--- a/Engine/Plugins/Experimental/Animation/SkeletalMeshModelingTools/Source/SkeletalMeshModelingTools/Private/SkeletalMeshModelingToolsModule.h
+++ b/Engine/Plugins/Experimental/Animation/SkeletalMeshModelingTools/Source/SkeletalMeshModelingTools/Private/SkeletalMeshModelingToolsModule.h
@@ -31,6 +31,7 @@ private:
 
 	bool IsEditingToolModeActive(TWeakPtr<ISkeletalMeshEditor> InSkeletalMeshEditor) const;
 	void OnToggleEditingToolsMode(TWeakPtr<ISkeletalMeshEditor> InSkeletalMeshEditor);
+	void CheckEnableEditingToolModeOnOpen(TWeakPtr<ISkeletalMeshEditor> InSkeletalMeshEditor);
 
 	void OnPostEngineInit();
 	
