Skeletal Mesh Modeling Plugin: Set the editing mode to default to 'on'

#rnx

diff --git a/Engine/Plugins/Experimental/Animation/SkeletalMeshModelingTools/Source/SkeletalMeshModelingTools/Private/SkeletalMeshModelingToolsModule.cpp b/Engine/Plugins/Experimental/Animation/SkeletalMeshModelingTools/Source/SkeletalMeshModelingTools/Private/SkeletalMeshModelingToolsModule.cpp
--- a/Engine/Plugins/Experimental/Animation/SkeletalMeshModelingTools/Source/SkeletalMeshModelingTools/Private/SkeletalMeshModelingToolsModule.cpp
+++ b/Engine/Plugins/Experimental/Animation/SkeletalMeshModelingTools/Source/SkeletalMeshModelingTools/Private/SkeletalMeshModelingToolsModule.cpp
@@ -188,7 +188,7 @@ void FSkeletalMeshModelingToolsModule::OnToggleEditingToolsMode(TWeakPtr<ISkelet
 
 void FSkeletalMeshModelingToolsModule::CheckEnableEditingToolModeOnOpen(TWeakPtr<ISkeletalMeshEditor> InSkeletalMeshEditor)
 {
-	bool bEditingModeActive = false;
+	bool bEditingModeActive = true;
 	GConfig->GetBool(ConfigSection, ConfigEditingModeActiveKey, bEditingModeActive, GEditorPerProjectIni);
 	if (bEditingModeActive && !IsEditingToolModeActive(InSkeletalMeshEditor))
 	{
