Skeleton Editing Tool: remove useless SSkeletonTree notifications once switching to the SRefSkeletonTree

#rnx

diff --git a/Engine/Plugins/Experimental/Animation/SkeletalMeshModelingTools/Source/SkeletalMeshModelingTools/Private/SkeletalMeshEditorUtils.cpp b/Engine/Plugins/Experimental/Animation/SkeletalMeshModelingTools/Source/SkeletalMeshModelingTools/Private/SkeletalMeshEditorUtils.cpp
--- a/Engine/Plugins/Experimental/Animation/SkeletalMeshModelingTools/Source/SkeletalMeshModelingTools/Private/SkeletalMeshEditorUtils.cpp
+++ b/Engine/Plugins/Experimental/Animation/SkeletalMeshModelingTools/Source/SkeletalMeshModelingTools/Private/SkeletalMeshEditorUtils.cpp
@@ -259,15 +259,25 @@ void USkeletalMeshEditorContextObject::BindRefSkeletonTree(ISkeletalMeshEditingI
 			.Modifier(Modifier)
 	];
 
+	// switch between SSkeletonTree and SReferenceSkeletonTree
 	SkeletonTab->SetContent(RefSkeletonWidget.ToSharedRef());
 
 	FBindData BindData;
 	Tie(BindData.ToToolNotifierHandle, BindData.FromToolNotifierHandle) = BindInterfaceTo(InEditingInterface, RefSkeletonTree->GetNotifier());
 	TreeBindings.Emplace(InEditingInterface, BindData);
 
+	// unbind notifications on the initial SSkeletonTree to avoid useless notifications
 	if(const TSharedPtr<ISkeletalMeshEditorBinding> Binding = Editor.Pin()->GetBinding())
 	{
-		RefSkeletonTree->GetNotifier().HandleNotification(Binding->GetSelectedBones(), ESkeletalMeshNotifyType::BonesSelected);
+		const TArray<FName> SelectedBones = Binding->GetSelectedBones();
+		
+		if (const FBindData* EditorBindData = EditorBindings.Find(InEditingInterface))
+		{
+			UnbindInterfaceFrom(InEditingInterface, Binding->GetNotifier(), *EditorBindData);
+			Binding->GetNotifier().HandleNotification({}, ESkeletalMeshNotifyType::BonesSelected);
+		}
+
+		RefSkeletonTree->GetNotifier().HandleNotification(SelectedBones, ESkeletalMeshNotifyType::BonesSelected);
 	}
 }
 
@@ -275,8 +285,10 @@ void USkeletalMeshEditorContextObject::UnbindRefSkeletonTree(ISkeletalMeshEditin
 {
 	if (const FBindData* BindData = TreeBindings.Find(InEditingInterface))
 	{
+		TArray<FName> SelectedBones;
 		if (RefSkeletonTree.IsValid())
 		{
+			RefSkeletonTree->GetSelectedBoneNames(SelectedBones);
 			UnbindInterfaceFrom(InEditingInterface, RefSkeletonTree->GetNotifier(), *BindData);
 		}
 
@@ -291,6 +303,14 @@ void USkeletalMeshEditorContextObject::UnbindRefSkeletonTree(ISkeletalMeshEditin
 					SkeletonTab->SetContent(DefaultSkeletonWidget.IsValid() ? DefaultSkeletonWidget.ToSharedRef() : SNullWidget::NullWidget);
 				}
 			}
+
+			// update the SSkeletonTree with the new RefSkeleton and set current selection 
+			if (const TSharedPtr<ISkeletalMeshEditorBinding> Binding = Editor.Pin()->GetBinding())
+			{
+				ISkeletalMeshNotifier& Notifier = Binding->GetNotifier();
+				Notifier.HandleNotification(SelectedBones, ESkeletalMeshNotifyType::HierarchyChanged);
+				Notifier.HandleNotification(SelectedBones, ESkeletalMeshNotifyType::BonesSelected);
+			}
 		}
 		
 		RefSkeletonTree.Reset();
diff --git a/Engine/Plugins/Experimental/MeshModelingToolsetExp/Source/MeshModelingToolsEditorOnlyExp/Private/SkeletalMesh/SkeletonEditingTool.cpp b/Engine/Plugins/Experimental/MeshModelingToolsetExp/Source/MeshModelingToolsEditorOnlyExp/Private/SkeletalMesh/SkeletonEditingTool.cpp
--- a/Engine/Plugins/Experimental/MeshModelingToolsetExp/Source/MeshModelingToolsEditorOnlyExp/Private/SkeletalMesh/SkeletonEditingTool.cpp
+++ b/Engine/Plugins/Experimental/MeshModelingToolsetExp/Source/MeshModelingToolsEditorOnlyExp/Private/SkeletalMesh/SkeletonEditingTool.cpp
@@ -1675,7 +1675,7 @@ void USkeletonEditingTool::HandleSkeletalMeshModified(const TArray<FName>& InBon
 
 	TArray<FName> BoneNames(InBoneNames);
 	const FReferenceSkeleton& RefSkeleton = Modifier->GetReferenceSkeleton();
-	BoneNames.RemoveAll([&](const FName& BoneName)
+	BoneNames.RemoveAll([&RefSkeleton](const FName& BoneName)
 	{
 		return RefSkeleton.FindRawBoneIndex(BoneName) == INDEX_NONE;
 	});
@@ -1686,9 +1686,9 @@ void USkeletonEditingTool::HandleSkeletalMeshModified(const TArray<FName>& InBon
 			Selection = BoneNames;
 			break;
 		case ESkeletalMeshNotifyType::BonesRemoved:
-			Selection.RemoveAll([&](const FName& BoneName)
+			Selection.RemoveAll([&InBoneNames](const FName& BoneName)
 			{
-				return BoneNames.Contains(BoneName);
+				return InBoneNames.Contains(BoneName);
 			});
 			break;
 		case ESkeletalMeshNotifyType::BonesMoved:
