Moved skin paint detail customization into SkeletalMeshModelingTools module (trivial)
This is preparation to use common editor styling in details panel for component selection.

diff --git a/Engine/Plugins/Experimental/Animation/SkeletalMeshModelingTools/Source/SkeletalMeshModelingTools/Private/DetailCustomization/SkinWeightDetailCustomization.cpp b/Engine/Plugins/Experimental/Animation/SkeletalMeshModelingTools/Source/SkeletalMeshModelingTools/Private/DetailCustomization/SkinWeightDetailCustomization.cpp
new file mode 100644
--- /dev/null
+++ b/Engine/Plugins/Experimental/Animation/SkeletalMeshModelingTools/Source/SkeletalMeshModelingTools/Private/DetailCustomization/SkinWeightDetailCustomization.cpp
@@ -0,0 +1,819 @@
+﻿// Copyright Epic Games, Inc. All Rights Reserved.
+
+#include "SkinWeightDetailCustomization.h"
+
+#include "DetailCategoryBuilder.h"
+#include "DetailLayoutBuilder.h"
+#include "DetailWidgetRow.h"
+#include "Widgets/Input/SSegmentedControl.h"
+#include "SkeletalMesh/SkinWeightsPaintTool.h"
+#include "SSkinWeightProfileImportOptions.h"
+#include "Selection/PolygonSelectionMechanic.h"
+#include "UObject/UnrealTypePrivate.h"
+#include "Widgets/Input/SSpinBox.h"
+#include "Widgets/Input/SButton.h"
+#include "Widgets/Input/SNumericEntryBox.h"
+#include "Widgets/Input/SSlider.h"
+
+#define LOCTEXT_NAMESPACE "SkinWeightToolSettingsEditor"
+
+// layout constants
+float FSkinWeightDetailCustomization::WeightSliderWidths = 150.0f;
+float FSkinWeightDetailCustomization::WeightEditingLabelsPercent = 0.40f;
+float FSkinWeightDetailCustomization::WeightEditVerticalPadding = 4.0f;
+float FSkinWeightDetailCustomization::WeightEditHorizontalPadding = 2.0f;
+
+void FSkinWeightDetailCustomization::CustomizeDetails(IDetailLayoutBuilder& DetailBuilder)
+{
+	CurrentDetailBuilder = &DetailBuilder;
+	
+	TArray<TWeakObjectPtr<UObject>> DetailObjects;
+	DetailBuilder.GetObjectsBeingCustomized(DetailObjects);
+
+	// should be impossible to get multiple settings objects for a single tool
+	ensure(DetailObjects.Num()==1);
+	SkinToolSettings = Cast<USkinWeightsPaintToolProperties>(DetailObjects[0]);
+
+	// custom display of falloff mode as segmented toggle buttons
+	IDetailCategoryBuilder& EditModeCategory = DetailBuilder.EditCategory("Weight Editing Mode", FText::GetEmpty(), ECategoryPriority::Important);
+
+	// add segmented control toggle for editing modes ("Brush" or "Selection")
+	EditModeCategory.AddCustomRow(LOCTEXT("EditModeCategory", "Weight Editing Mode"), false)
+	.WholeRowContent()
+	[
+		SNew(SBox)
+		.Padding(2.0f)
+		[
+			SNew(SSegmentedControl<EWeightEditMode>)
+			.ToolTipText(LOCTEXT("EditingModeTooltip",
+					"Brush: edit weights by painting directly on mesh.\n"
+					"Vertices: select vertices and edit weights directly.\n"
+					"Bones: select and manipulate bones to preview deformations.\n"))
+			.Value_Lambda([this]()
+			{
+				return SkinToolSettings->EditingMode;
+			})
+			.OnValueChanged_Lambda([this](EWeightEditMode Mode)
+			{
+				SkinToolSettings->EditingMode = Mode;
+				SkinToolSettings->WeightTool->ToggleEditingMode();
+				if (CurrentDetailBuilder)
+				{
+					CurrentDetailBuilder->ForceRefreshDetails();
+				}
+			})
+			+SSegmentedControl<EWeightEditMode>::Slot(EWeightEditMode::Brush)
+			.Text(LOCTEXT("BrushEditMode", "Brush"))
+			+ SSegmentedControl<EWeightEditMode>::Slot(EWeightEditMode::Vertices)
+			.Text(LOCTEXT("VertexEditMode", "Vertices"))
+			+ SSegmentedControl<EWeightEditMode>::Slot(EWeightEditMode::Bones)
+			.Text(LOCTEXT("BoneEditMode", "Bones"))
+		]
+	];
+
+	// BRUSH editing mode UI
+	if (SkinToolSettings->EditingMode == EWeightEditMode::Brush)
+	{
+		AddBrushUI(DetailBuilder);
+	}
+
+	// VERTEX editing mode UI
+	if (SkinToolSettings->EditingMode == EWeightEditMode::Vertices)
+	{
+		AddSelectionUI(DetailBuilder);
+	}
+	
+	// COLOR MODE category
+	IDetailCategoryBuilder& MeshDisplayCategory = DetailBuilder.EditCategory("MeshDisplay", FText::GetEmpty(), ECategoryPriority::Important);
+	MeshDisplayCategory.InitiallyCollapsed(false);
+	MeshDisplayCategory.AddCustomRow(LOCTEXT("ColorModeCategory", "Color Mode"), false)
+	.WholeRowContent()
+	[
+		SNew(SBox)
+		.Padding(2.0f)
+		.HAlign(HAlign_Center)
+		[
+			SNew(SSegmentedControl<EWeightColorMode>)
+			.ToolTipText(LOCTEXT("WeightColorTooltip",
+					"Adjust the weight display in the viewport.\n\n"
+					"Greyscale: Displays weights on the current bone by blending from black (0) to white (1).\n"
+					"Ramp: Displays weights on the current bone. Weights at 0 and 1 use the min and max colors. Weights inbetween 0 and 1 use the ramp colors.\n"
+					"Multi Color: Displays weights on ALL bones using the color of the bones.\n"
+					"Full Material: Displays normal mesh materials with textures.\n"))
+			.Value_Lambda([this]()
+			{
+				return SkinToolSettings->ColorMode;
+			})
+			.OnValueChanged_Lambda([this](EWeightColorMode Mode)
+			{
+				SkinToolSettings->SetColorMode(Mode);
+			})
+			+ SSegmentedControl<EWeightColorMode>::Slot(EWeightColorMode::Greyscale)
+			.Text(LOCTEXT("GreyscaleMode", "Greyscale"))
+			+SSegmentedControl<EWeightColorMode>::Slot(EWeightColorMode::Ramp)
+			.Text(LOCTEXT("RampMode", "Ramp"))
+			+SSegmentedControl<EWeightColorMode>::Slot(EWeightColorMode::BoneColors)
+			.Text(LOCTEXT("BoneColorsMode", "Bone Colors"))
+			+SSegmentedControl<EWeightColorMode>::Slot(EWeightColorMode::FullMaterial)
+			.Text(LOCTEXT("MaterialMode", "Full Material"))
+		]
+	];
+
+	// hide all base brush properties that have been customized
+	const TSharedRef<IPropertyHandle> BrushModeHandle = DetailBuilder.GetProperty(GET_MEMBER_NAME_CHECKED(USkinWeightsPaintToolProperties, BrushMode));
+	DetailBuilder.HideProperty(BrushModeHandle);
+	const TSharedRef<IPropertyHandle> BrushSizeHandle = DetailBuilder.GetProperty(GET_MEMBER_NAME_CHECKED(UBrushBaseProperties, BrushSize), UBrushBaseProperties::StaticClass());
+	DetailBuilder.HideProperty(BrushSizeHandle);
+	const TSharedRef<IPropertyHandle> BrushStrengthHandle = DetailBuilder.GetProperty(GET_MEMBER_NAME_CHECKED(UBrushBaseProperties, BrushStrength), UBrushBaseProperties::StaticClass());
+	DetailBuilder.HideProperty(BrushStrengthHandle);
+	const TSharedRef<IPropertyHandle> BrushFalloffHandle = DetailBuilder.GetProperty(GET_MEMBER_NAME_CHECKED(UBrushBaseProperties, BrushFalloffAmount), UBrushBaseProperties::StaticClass());
+	DetailBuilder.HideProperty(BrushFalloffHandle);
+	const TSharedRef<IPropertyHandle> BrushRadiusHandle = DetailBuilder.GetProperty(GET_MEMBER_NAME_CHECKED(UBrushBaseProperties, BrushRadius), UBrushBaseProperties::StaticClass());
+	DetailBuilder.HideProperty(BrushRadiusHandle);
+	const TSharedRef<IPropertyHandle> SpecifyRadiusHandle = DetailBuilder.GetProperty(GET_MEMBER_NAME_CHECKED(UBrushBaseProperties, bSpecifyRadius), UBrushBaseProperties::StaticClass());
+	DetailBuilder.HideProperty(SpecifyRadiusHandle);
+	const TSharedRef<IPropertyHandle> EditModePropHandle = DetailBuilder.GetProperty(GET_MEMBER_NAME_CHECKED(USkinWeightsPaintToolProperties, EditingMode));
+	DetailBuilder.HideProperty(EditModePropHandle);
+	const TSharedRef<IPropertyHandle> ColorModePropHandle = DetailBuilder.GetProperty(GET_MEMBER_NAME_CHECKED(USkinWeightsPaintToolProperties, ColorMode));
+	DetailBuilder.HideProperty(ColorModePropHandle);
+}
+
+void FSkinWeightDetailCustomization::AddBrushUI(IDetailLayoutBuilder& DetailBuilder)
+{
+	// custom display of falloff mode as segmented toggle buttons
+	IDetailCategoryBuilder& BrushCategory = DetailBuilder.EditCategory("Brush", FText::GetEmpty(), ECategoryPriority::Important);
+
+	// add segmented control toggle for brush behavior modes ("Add", "Replace", etc..)
+	BrushCategory.AddCustomRow(LOCTEXT("BrushModeCategory", "Brush Mode"), false)
+	.WholeRowContent()
+	[
+		SNew(SBox)
+		.Padding(2.0f)
+		[
+			SNew(SSegmentedControl<EWeightEditOperation>)
+			.ToolTipText(LOCTEXT("FloodTooltip",
+				"Add: applies the current weight plus the flood value to the new weight.\n"
+				"Replace: applies the current weight minus the strength value to the new weight.\n"
+				"Multiply: applies the current weight multiplied by the strength value to the new weight.\n"
+				"Relax: applies the average of the connected (by edge) vertex weights to the new vertex weight, blended by the strength.\n"
+				"This command operates on the selected bone(s) and selected vertices.\n"
+				"If no bones are selected, ALL bones are considered.\n"
+				"If no vertices are selected, ALL vertices are considered."))
+			.Value_Lambda([this]()
+			{
+				return SkinToolSettings->BrushMode;
+			})
+			.OnValueChanged_Lambda([this](EWeightEditOperation Mode)
+			{
+				SkinToolSettings->SetBrushMode(Mode);
+			})
+			+SSegmentedControl<EWeightEditOperation>::Slot(EWeightEditOperation::Add)
+			.Text(LOCTEXT("BrushAddMode", "Add"))
+			+ SSegmentedControl<EWeightEditOperation>::Slot(EWeightEditOperation::Replace)
+			.Text(LOCTEXT("BrushReplaceMode", "Replace"))
+			+ SSegmentedControl<EWeightEditOperation>::Slot(EWeightEditOperation::Multiply)
+			.Text(LOCTEXT("BrushMultiplyMode", "Multiply"))
+			+ SSegmentedControl<EWeightEditOperation>::Slot(EWeightEditOperation::Relax)
+			.Text(LOCTEXT("BrushRelaxMode", "Relax"))
+		]
+	];
+
+	// add segmented control toggle for brush falloff modes ("Surface" or "Volume")
+	BrushCategory.AddCustomRow(LOCTEXT("BrushFalloffModeCategory", "Brush Falloff Mode"), false)
+	.WholeRowContent()
+	[
+		SNew(SBox)
+		.Padding(2.0f)
+		[
+			SNew(SSegmentedControl<EWeightBrushFalloffMode>)
+			.ToolTipText(LOCTEXT("BrushFalloffModeTooltip",
+					"Surface: falloff is based on the distance along the surface from the brush center to nearby connected vertices.\n"
+					"Volume: falloff is based on the straight-line distance from the brush center to surrounding vertices.\n"))
+			.Value_Lambda([this]()
+			{
+				return SkinToolSettings->GetBrushConfig().FalloffMode;
+			})
+			.OnValueChanged_Lambda([this](EWeightBrushFalloffMode Mode)
+			{
+				SkinToolSettings->SetFalloffMode(Mode);
+			})
+			+SSegmentedControl<EWeightBrushFalloffMode>::Slot(EWeightBrushFalloffMode::Surface)
+			.Text(LOCTEXT("SurfaceMode", "Surface"))
+			+ SSegmentedControl<EWeightBrushFalloffMode>::Slot(EWeightBrushFalloffMode::Volume)
+			.Text(LOCTEXT("VolumeMode", "Volume"))
+		]
+	];
+
+	BrushCategory.AddCustomRow(LOCTEXT("BrushSizeCategory", "Brush Radius"), false)
+	.NameContent()
+	[
+		SNew(STextBlock)
+		.Text(LOCTEXT("BrushRadiusLabel", "Radius"))
+		.Font(FAppStyle::GetFontStyle(TEXT("PropertyWindow.NormalFont")))
+		.ToolTipText(LOCTEXT("BrushRadiusTooltip", "The radius of the brush in scene units."))
+	]
+	.ValueContent()
+	[
+		SNew(SSpinBox<float>)
+		.MinValue(0.01f)
+		.MaxSliderValue(20.f)
+		.Value(10.0f)
+		.SupportDynamicSliderMaxValue(true)
+		.Value_Lambda([this]()
+		{
+			return SkinToolSettings->GetBrushConfig().Radius;
+		})
+		.OnValueChanged_Lambda([this](float NewValue)
+		{
+			SkinToolSettings->BrushRadius = NewValue;
+			SkinToolSettings->GetBrushConfig().Radius = NewValue;
+			FPropertyChangedEvent PropertyChangedEvent(UBrushBaseProperties::StaticClass()->FindPropertyByName(GET_MEMBER_NAME_CHECKED(UBrushBaseProperties, BrushRadius)));
+			SkinToolSettings->PostEditChangeProperty(PropertyChangedEvent);
+		})
+		.OnValueCommitted_Lambda([this](float NewValue, ETextCommit::Type CommitType)
+		{
+			SkinToolSettings->SaveConfig();
+		})
+	];
+
+	BrushCategory.AddCustomRow(LOCTEXT("BrushStrengthCategory", "Brush Strength"), false)
+	.NameContent()
+	[
+		SNew(STextBlock)
+		.Text(LOCTEXT("BrushStrengthLabel", "Strength"))
+		.Font(FAppStyle::GetFontStyle(TEXT("PropertyWindow.NormalFont")))
+		.ToolTipText(LOCTEXT("BrushStrengthTooltip", "The strength of the effect on the weights. Exact effect depends on the Brush mode."))
+	]
+	.ValueContent()
+	[
+		SNew(SSpinBox<float>)
+		.MinValue(0.f)
+		.MaxValue(2.0f)
+		.MaxSliderValue(1.f)
+		.Value(1.0f)
+		.SupportDynamicSliderMaxValue(true)
+		.Value_Lambda([this]()
+		{
+			return SkinToolSettings->GetBrushConfig().Strength;
+		})
+		.OnValueChanged_Lambda([this](float NewValue)
+		{
+			SkinToolSettings->BrushStrength = NewValue;
+			SkinToolSettings->GetBrushConfig().Strength = NewValue;
+			FPropertyChangedEvent PropertyChangedEvent(UBrushBaseProperties::StaticClass()->FindPropertyByName(GET_MEMBER_NAME_CHECKED(UBrushBaseProperties, BrushStrength)));
+			SkinToolSettings->PostEditChangeProperty(PropertyChangedEvent);
+		})
+		.OnValueCommitted_Lambda([this](float NewValue, ETextCommit::Type CommitType)
+		{
+			SkinToolSettings->SaveConfig();
+		})
+	];
+
+	BrushCategory.AddCustomRow(LOCTEXT("BrushFalloffCategory", "Brush Falloff"), false)
+	.NameContent()
+	[
+		SNew(STextBlock)
+		.Text(LOCTEXT("BrushFalloffLabel", "Falloff"))
+		.Font(FAppStyle::GetFontStyle(TEXT("PropertyWindow.NormalFont")))
+		.ToolTipText(LOCTEXT("BrushFalloffTooltip", "At 0, the brush has no falloff. At 1 it has exponential falloff."))
+	]
+	.ValueContent()
+	[
+		SNew(SSpinBox<float>)
+		.MinValue(0.f)
+		.MaxValue(1.f)
+		.Value_Lambda([this]()
+		{
+			return SkinToolSettings->GetBrushConfig().Falloff;
+		})
+		.OnValueChanged_Lambda([this](float NewValue)
+		{
+			SkinToolSettings->BrushFalloffAmount = NewValue;
+			SkinToolSettings->GetBrushConfig().Falloff = NewValue;
+			FPropertyChangedEvent PropertyChangedEvent(UBrushBaseProperties::StaticClass()->FindPropertyByName(GET_MEMBER_NAME_CHECKED(UBrushBaseProperties, BrushFalloffAmount)));
+			SkinToolSettings->PostEditChangeProperty(PropertyChangedEvent);
+		})
+		.OnValueCommitted_Lambda([this](float NewValue, ETextCommit::Type CommitType)
+		{
+			SkinToolSettings->SaveConfig();
+		})
+	];
+}
+
+void FSkinWeightDetailCustomization::AddSelectionUI(IDetailLayoutBuilder& DetailBuilder)
+{
+	// custom display of weight editing tools
+	IDetailCategoryBuilder& EditSelectionCategory = DetailBuilder.EditCategory("Edit Selection", FText::GetEmpty(), ECategoryPriority::Important);
+	EditSelectionCategory.InitiallyCollapsed(true);
+
+	// GROW/SHRINK/FLOOD Selection category
+	EditSelectionCategory.AddCustomRow(LOCTEXT("EditSelectionRow", "Edit Selection"), false)
+	.WholeRowContent()
+	[
+		SNew(SHorizontalBox)
+
+		+SHorizontalBox::Slot()
+		.Padding(2.f, WeightEditVerticalPadding)
+		[
+			SNew(SButton)
+			.HAlign(HAlign_Center)
+			.VAlign(VAlign_Center)
+			.Text(LOCTEXT("GrowSelectionButtonLabel", "Grow"))
+			.ToolTipText(LOCTEXT("GrowSelectionTooltip",
+					"Grow the current selection by adding connected neighbors to current selection.\n"))
+			.OnClicked_Lambda([this]()
+			{
+				SkinToolSettings->WeightTool->GrowSelection();
+				return FReply::Handled();
+			})
+		]
+
+		+SHorizontalBox::Slot()
+		.Padding(2.f, WeightEditVerticalPadding)
+		[
+			SNew(SButton)
+			.HAlign(HAlign_Center)
+			.VAlign(VAlign_Center)
+			.Text(LOCTEXT("ShrinkSelectionButtonLabel", "Shrink"))
+			.ToolTipText(LOCTEXT("ShrinkSelectionTooltip",
+					"Shrink the current selection by removing vertices on the border of the current selection.\n"))
+			.OnClicked_Lambda([this]()
+			{
+				SkinToolSettings->WeightTool->ShrinkSelection();
+				return FReply::Handled();
+			})
+		]
+
+		+SHorizontalBox::Slot()
+		.Padding(2.f, WeightEditVerticalPadding)
+		[
+			SNew(SButton)
+			.HAlign(HAlign_Center)
+			.VAlign(VAlign_Center)
+			.Text(LOCTEXT("FloodSelectionButtonLabel", "Flood"))
+			.ToolTipText(LOCTEXT("FloodSelectionTooltip",
+					"Flood the current selection by adding all connected vertices to the current selection.\n"))
+			.OnClicked_Lambda([this]()
+			{
+				SkinToolSettings->WeightTool->FloodSelection();
+				return FReply::Handled();
+			})
+		]
+	];
+
+	// custom display of weight editing tools
+	IDetailCategoryBuilder& EditWeightsCategory = DetailBuilder.EditCategory("Edit Weights", FText::GetEmpty(), ECategoryPriority::Important);
+	EditWeightsCategory.InitiallyCollapsed(true);
+
+	// AVERAGE/RELAX/NORMALIZE WEIGHTS category
+	EditWeightsCategory.AddCustomRow(LOCTEXT("NormalizeWeightsRow", "Normalize"), false)
+	.WholeRowContent()
+	[
+		SNew(SHorizontalBox)
+
+		+SHorizontalBox::Slot()
+		.Padding(2.f, WeightEditVerticalPadding)
+		[
+			SNew(SButton)
+			.HAlign(HAlign_Center)
+			.VAlign(VAlign_Center)
+			.Text(LOCTEXT("AverageWeightsButtonLabel", "Average"))
+			.ToolTipText(LOCTEXT("AverageWeightsTooltip",
+					"Takes the average of vertex weights and applies it to the selected vertices.\n"
+					"This command operates on the selected bone(s) and selected vertices.\n "
+					"If no bones are selected, ALL bone weights are considered.\n "
+					"If no vertices are selected, ALL vertices are considered."))
+			.OnClicked_Lambda([this]()
+			{
+				SkinToolSettings->WeightTool->AverageWeights();
+				return FReply::Handled();
+			})
+		]
+
+		+SHorizontalBox::Slot()
+		.Padding(2.f, WeightEditVerticalPadding)
+		[
+			SNew(SButton)
+			.HAlign(HAlign_Center)
+			.VAlign(VAlign_Center)
+			.Text(LOCTEXT("NormalizeWeightsButtonLabel", "Normalize"))
+			.ToolTipText(LOCTEXT("NormalizeWeightsTooltip",
+					"Forces the weights on the selected vertices to sum to 1.\n"
+					"This command operates on the selected vertices.\n "
+					"If no vertices are selected, ALL vertices are considered."))
+			.IsEnabled_Lambda([this]()
+			{
+				return SkinToolSettings->EditingMode == EWeightEditMode::Vertices;
+			})
+			.OnClicked_Lambda([this]()
+			{
+				SkinToolSettings->WeightTool->NormalizeWeights();
+				return FReply::Handled();
+			})
+		]
+	];
+
+	// MIRROR WEIGHTS category
+	EditWeightsCategory.AddCustomRow(LOCTEXT("MirrorWeightsRow", "Mirror"), false)
+	.WholeRowContent()
+	[
+		SNew(SVerticalBox)
+		
+		+ SVerticalBox::Slot()
+		.Padding(WeightEditHorizontalPadding, WeightEditVerticalPadding)
+		[
+			SNew(SHorizontalBox)
+			
+			+SHorizontalBox::Slot()
+			.VAlign(VAlign_Center)
+			.FillWidth(WeightEditingLabelsPercent)
+			[
+				SNew(STextBlock)
+				.Text(LOCTEXT("MirrorPlaneLabel", "Mirror Plane"))
+				.Font(FAppStyle::GetFontStyle(TEXT("PropertyWindow.NormalFont")))
+				.ToolTipText(LOCTEXT("MirrorPlaneTooltip", "The plane to copy weights across."))
+			]
+
+			+SHorizontalBox::Slot()
+			.FillWidth(1.f)
+			[
+				SNew(SHorizontalBox)
+				+SHorizontalBox::Slot()
+				[
+					SNew(SSegmentedControl<EAxis::Type>)
+					.ToolTipText(LOCTEXT("MirrorAxisTooltip",
+						"X: copies weights across the YZ plane.\n"
+						"Y: copies weights across the XZ plane.\n"
+						"Z: copies weights across the XY plane."))
+					.Value_Lambda([this]()
+					{
+						return SkinToolSettings->MirrorAxis;
+					})
+					.OnValueChanged_Lambda([this](EAxis::Type Mode)
+					{
+						SkinToolSettings->MirrorAxis = Mode;
+					})
+					+ SSegmentedControl<EAxis::Type>::Slot(EAxis::X)
+					.Text(LOCTEXT("MirrorXLabel", "X"))
+					+ SSegmentedControl<EAxis::Type>::Slot(EAxis::Y)
+					.Text(LOCTEXT("MirrorYLabel", "Y"))
+					+ SSegmentedControl<EAxis::Type>::Slot(EAxis::Z)
+					.Text(LOCTEXT("MirrorZLabel", "Z"))
+				]
+					
+				+SHorizontalBox::Slot()
+				[
+					SNew(SSegmentedControl<EMirrorDirection>)
+					.ToolTipText(LOCTEXT("MirrorDirectionTooltip", "The direction that determines what side of the plane to copy weights from."))
+					.Value_Lambda([this]()
+					{
+						return SkinToolSettings->MirrorDirection;
+					})
+					.OnValueChanged_Lambda([this](EMirrorDirection Mode)
+					{
+						SkinToolSettings->MirrorDirection = Mode;
+					})
+					+ SSegmentedControl<EMirrorDirection>::Slot(EMirrorDirection::PositiveToNegative)
+					.Text(LOCTEXT("MirrorPosToNegLabel", "+ to -"))
+					+ SSegmentedControl<EMirrorDirection>::Slot(EMirrorDirection::NegativeToPositive)
+					.Text(LOCTEXT("MirrorNegToPosLabel", "- to +"))
+				]
+			]
+		]
+
+		+SVerticalBox::Slot()
+		.Padding(0.f, WeightEditVerticalPadding)
+		[
+			SNew(SBox)
+			[
+				SNew(SButton)
+				.HAlign(HAlign_Center)
+				.Text(LOCTEXT("MirrorWeightsButtonLabel", "Mirror"))
+				.ToolTipText(LOCTEXT("MirrorButtonTooltip",
+					"Weights are copied across the given plane in the given direction.\n"
+					"This command operates on the selected bone(s) and selected vertices.\n "
+					"If no bones are selected, ALL bone weights are considered.\n "
+					"If no vertices are selected, ALL vertices are considered."))
+				.OnClicked_Lambda([this]()
+				{
+					SkinToolSettings->WeightTool->MirrorWeights(SkinToolSettings->MirrorAxis, SkinToolSettings->MirrorDirection);
+					return FReply::Handled();
+				})
+			]
+		]
+	];
+
+	// FLOOD WEIGHTS category
+	EditWeightsCategory.AddCustomRow(LOCTEXT("FloodWeightsRow", "Flood"), false)
+	.WholeRowContent()
+	[
+		SNew(SVerticalBox)
+
+		+ SVerticalBox::Slot()
+		.Padding(0.f, WeightEditVerticalPadding)
+		[
+			SNew(SHorizontalBox)
+				
+			+SHorizontalBox::Slot()
+			.VAlign(VAlign_Center)
+			.FillWidth(WeightEditingLabelsPercent)
+			[
+				SNew(STextBlock)
+				.Text(LOCTEXT("FloodAmountLabel", "Flood Amount"))
+				.Font(FAppStyle::GetFontStyle(TEXT("PropertyWindow.NormalFont")))
+				.ToolTipText(LOCTEXT("FloodAmountTooltip", "The amount of weight to apply in Flood operation."))
+			]
+
+			+SHorizontalBox::Slot()
+			.MaxWidth(WeightSliderWidths)
+			[
+				SNew(SSpinBox<float>)
+				.MinValue(0.f)
+				.MaxValue(2.0f)
+				.MaxSliderValue(1.f)
+				.Value(1.0f)
+				.SupportDynamicSliderMaxValue(true)
+				.Value_Lambda([this]()
+				{
+					return SkinToolSettings->FloodValue;
+				})
+				.OnValueChanged_Lambda([this](float NewValue)
+				{
+					SkinToolSettings->FloodValue = NewValue;
+				})
+				.OnValueCommitted_Lambda([this](float NewValue, ETextCommit::Type CommitType)
+				{
+					SkinToolSettings->SaveConfig();
+				})
+			]
+		]
+
+		+SVerticalBox::Slot()
+		.Padding(0.f, WeightEditVerticalPadding)
+		[
+			SNew(SHorizontalBox)
+			
+			+SHorizontalBox::Slot()
+			.VAlign(VAlign_Center)
+			.FillWidth(WeightEditingLabelsPercent)
+			[
+				SNew(STextBlock)
+				.Text(LOCTEXT("FloodOperationLabel", "Flood Operation"))
+				.Font(FAppStyle::GetFontStyle(TEXT("PropertyWindow.NormalFont")))
+				.ToolTipText(LOCTEXT("FloodOperationTooltip", "The various flood weight operations."))
+			]
+
+			+SHorizontalBox::Slot()
+			[
+				SNew(SHorizontalBox)
+				+ SHorizontalBox::Slot()
+				[
+					SNew(SButton)
+					.HAlign(HAlign_Center)
+					.Text(LOCTEXT("AddWeightsButtonLabel", "Add"))
+					.ToolTipText(LOCTEXT("AddOpTooltip", "Add: applies the current weight plus the flood amount to the new weight."))
+					.OnClicked_Lambda([this]()
+					{
+						SkinToolSettings->WeightTool->FloodWeights(SkinToolSettings->FloodValue, EWeightEditOperation::Add);
+						return FReply::Handled();
+					})
+				]
+				+ SHorizontalBox::Slot()
+				[
+					SNew(SButton)
+					.HAlign(HAlign_Center)
+					.Text(LOCTEXT("ReplaceWeightsButtonLabel", "Replace"))
+					.ToolTipText(LOCTEXT("ReplaceOpTooltip", "Replace: applies the current weight minus the flood amount to the new weight."))
+					.OnClicked_Lambda([this]()
+					{
+						SkinToolSettings->WeightTool->FloodWeights(SkinToolSettings->FloodValue, EWeightEditOperation::Replace);
+						return FReply::Handled();
+					})
+				]
+				+ SHorizontalBox::Slot()
+				[
+					SNew(SButton)
+					.HAlign(HAlign_Center)
+					.Text(LOCTEXT("MultiplyeightsButtonLabel", "Multiply"))
+					.ToolTipText(LOCTEXT("MultiplyOpTooltip", "Multiply: applies the current weight multiplied by the flood amount to the new weight."))
+					.OnClicked_Lambda([this]()
+					{
+						SkinToolSettings->WeightTool->FloodWeights(SkinToolSettings->FloodValue, EWeightEditOperation::Multiply);
+						return FReply::Handled();
+					})
+				]
+				+ SHorizontalBox::Slot()
+				[
+					SNew(SButton)
+					.HAlign(HAlign_Center)
+					.Text(LOCTEXT("RelaxWeightsButtonLabel", "Relax"))
+					.ToolTipText(LOCTEXT("RelaxeOpTooltip", "Relax: applies the average of the connected (by edge) vertex weights to the new vertex weight, scaled by the flood amount."))
+					.OnClicked_Lambda([this]()
+					{
+						SkinToolSettings->WeightTool->FloodWeights(SkinToolSettings->FloodValue, EWeightEditOperation::Relax);
+						return FReply::Handled();
+					})
+				]
+			]
+		]
+	];
+
+	// PRUNE WEIGHTS category
+	EditWeightsCategory.AddCustomRow(LOCTEXT("PruneWeightsRow", "Prune"), false)
+	.WholeRowContent()
+	[
+		SNew(SVerticalBox)
+
+		+ SVerticalBox::Slot()
+		.Padding(WeightEditHorizontalPadding, WeightEditVerticalPadding)
+		[
+			SNew(SHorizontalBox)
+					
+			+SHorizontalBox::Slot()
+			.VAlign(VAlign_Center)
+			.FillWidth(WeightEditingLabelsPercent)
+			[
+				SNew(STextBlock)
+				.Text(LOCTEXT("PruneThresholdLabel", "Prune Threshold"))
+				.Font(FAppStyle::GetFontStyle(TEXT("PropertyWindow.NormalFont")))
+				.ToolTipText(LOCTEXT("PruneThresholdTooltip", "The threshold weight value to use with Prune operation."))
+			]
+
+			+SHorizontalBox::Slot()
+			.MaxWidth(WeightSliderWidths)
+			[
+				SNew(SSpinBox<float>)
+				.MinValue(0.f)
+				.MaxValue(1.f)
+				.Value_Lambda([this]()
+				{
+					return SkinToolSettings->PruneValue;
+				})
+				.OnValueChanged_Lambda([this](float NewValue)
+				{
+					SkinToolSettings->PruneValue = NewValue;
+				})
+				.OnValueCommitted_Lambda([this](float NewValue, ETextCommit::Type CommitType)
+				{
+					SkinToolSettings->SaveConfig();
+				})
+			]
+		]
+		
+		+SVerticalBox::Slot()
+		.Padding(WeightEditHorizontalPadding, WeightEditVerticalPadding)
+		[
+			SNew(SBox)
+			.MinDesiredWidth(WeightSliderWidths)
+			[
+				SNew(SButton)
+				.HAlign(HAlign_Center)
+				.Text(LOCTEXT("PruneWeightsButtonLabel", "Prune"))
+				.ToolTipText(LOCTEXT("PruneButtonTooltip",
+					"Weights below the given threshold value are removed.\n"
+					"This command operates on the selected bone(s) and selected vertices.\n "
+					"If no bones are selected, ALL bone weights are considered.\n "
+					"If no vertices are selected, ALL vertices are considered."))
+				.OnClicked_Lambda([this]()
+				{
+					SkinToolSettings->WeightTool->PruneWeights(SkinToolSettings->PruneValue);
+					return FReply::Handled();
+				})
+			]
+		]
+	];
+
+	// VERTEX EDITOR category
+	EditWeightsCategory.AddCustomRow(LOCTEXT("VertexEditorRow", "Vertex Editor"), false)
+	.WholeRowContent()
+	[
+		SNew(SVertexWeightEditor, SkinToolSettings->WeightTool)
+	];
+}
+
+void SVertexWeightItem::Construct(const FArguments& InArgs, const TSharedRef<STableViewBase>& OwnerTableView)
+{
+	Element = InArgs._Element;
+	ParentTable = InArgs._ParentTable;
+	SMultiColumnTableRow<TSharedPtr<FWeightEditorElement>>::Construct(FSuperRowType::FArguments(), OwnerTableView);
+}
+
+TSharedRef<SWidget> SVertexWeightItem::GenerateWidgetForColumn(const FName& ColumnName)
+{
+	if (ColumnName == "Bone")
+	{
+		const FName BoneName = ParentTable->Tool->GetBoneNameFromIndex(Element->BoneIndex);
+		return SNew(STextBlock).Text(FText::FromName(BoneName));
+	}
+
+	if (ColumnName == "Weight")
+	{
+		// add a buffer to the slider to prevent ever fully getting a value to 1 or 0 using the slider alone
+		// doing so will remove other influences and cause the slider to no longer function as all other influences
+		// will be culled by normalization thus making the slider "stuck" at full value.
+		constexpr float SliderBuffer = 0.001f;
+		
+		return SNew(SNumericEntryBox<float>)
+		.AllowSpin(true)
+		.MinSliderValue(SliderBuffer)
+		.MinValue(0.f)
+		.MaxSliderValue(1.0f-SliderBuffer)
+		.MaxValue(1.f)
+		.Value_Lambda([this]()
+		{
+			return ParentTable->Tool->GetAverageWeightOnBone(Element->BoneIndex, ParentTable->SelectedVertices);
+		})
+		.OnValueChanged_Lambda([this](float NewValue)
+		{
+			TArray<VertexIndex> VerticesToEdit;
+			ParentTable->Tool->GetSelectedVertices(VerticesToEdit);
+			ParentTable->Tool->SetBoneWeightOnVertices(Element->BoneIndex, NewValue, VerticesToEdit, !bInTransaction);
+		})
+		.OnValueCommitted_Lambda([this](float NewValue, ETextCommit::Type CommitType)
+		{
+			bInTransaction = false;
+		})
+		.OnBeginSliderMovement_Lambda([this]()
+		{
+			ParentTable->Tool->BeginChange();
+			bInTransaction = true;
+		})
+		.OnEndSliderMovement_Lambda([this](float)
+		{
+			const FText TransactionLabel = LOCTEXT("DirectWeightChange", "Set weights on vertices.");
+			ParentTable->Tool->EndChange(TransactionLabel);
+			bInTransaction = false;
+		})
+		.ToolTipText(LOCTEXT("WeightSliderToolTip", "Set the weight on this bone for the selected vertices."));
+	}
+
+	checkNoEntry();
+	return SNullWidget::NullWidget;
+}
+
+SVertexWeightEditor::~SVertexWeightEditor()
+{
+	if (Tool.IsValid())
+	{
+		Tool->OnSelectionChanged.RemoveAll(this);
+		Tool->OnWeightsChanged.RemoveAll(this);
+		Tool.Reset();
+	}
+}
+
+void SVertexWeightEditor::Construct(const FArguments& InArgs, USkinWeightsPaintTool* InSkinTool)
+{
+	Tool = InSkinTool;
+
+	ChildSlot
+	[
+		SNew(SBox)
+		[
+			SAssignNew( ListView, SWeightEditorListViewType )
+			.SelectionMode(ESelectionMode::Single)
+			.ListItemsSource( &ListViewItems )
+			.OnGenerateRow_Lambda([this](TSharedPtr<FWeightEditorElement> Element, const TSharedRef<STableViewBase>& OwnerTableView)
+			{
+				return SNew(SVertexWeightItem, OwnerTableView).Element(Element).ParentTable(SharedThis(this));
+			})
+			.HeaderRow
+			(
+				SNew(SHeaderRow)
+				+ SHeaderRow::Column("Bone").DefaultLabel(NSLOCTEXT("WeightEditorBoneColumn", "Bone", "Bone"))
+				+ SHeaderRow::Column("Weight").DefaultLabel(NSLOCTEXT("WeightEditorWeightColumn", "Weight (Average)", "Weight (Average)"))
+			)
+		]
+	];
+
+	RefreshView();
+	
+	Tool->OnSelectionChanged.AddSP(this, &SVertexWeightEditor::RefreshView);
+	Tool->OnWeightsChanged.AddSP(this, &SVertexWeightEditor::RefreshView);
+}
+
+void SVertexWeightEditor::RefreshView()
+{
+	if (!Tool.IsValid())
+	{
+		return; 
+	}
+
+	// get list of selected vertex indices
+	SelectedVertices.Reset();
+	Tool->GetSelectedVertices(SelectedVertices);
+	
+	// get all bones affecting the selected vertices
+	TArray<int32> Influences;
+	Tool->GetInfluences(SelectedVertices, Influences);
+
+	// generate list view items
+	ListViewItems.Reset();
+	for (const int32 InfluenceIndex : Influences)
+	{
+		ListViewItems.Add(MakeShareable(new FWeightEditorElement(InfluenceIndex)));
+	}
+	
+	ListView->RequestListRefresh();
+}
+
+#undef LOCTEXT_NAMESPACE
diff --git a/Engine/Plugins/Experimental/Animation/SkeletalMeshModelingTools/Source/SkeletalMeshModelingTools/Private/DetailCustomization/SkinWeightDetailCustomization.h b/Engine/Plugins/Experimental/Animation/SkeletalMeshModelingTools/Source/SkeletalMeshModelingTools/Private/DetailCustomization/SkinWeightDetailCustomization.h
new file mode 100644
--- /dev/null
+++ b/Engine/Plugins/Experimental/Animation/SkeletalMeshModelingTools/Source/SkeletalMeshModelingTools/Private/DetailCustomization/SkinWeightDetailCustomization.h
@@ -0,0 +1,92 @@
+﻿// Copyright Epic Games, Inc. All Rights Reserved.
+
+#pragma once
+
+#include "CoreMinimal.h"
+#include "IDetailCustomization.h"
+#include "UObject/WeakObjectPtr.h"
+#include "Widgets/Views/STableRow.h"
+#include "Widgets/Views/SListView.h"
+
+class USkinWeightsPaintTool;
+class USkinWeightsPaintToolProperties;
+
+class FSkinWeightDetailCustomization : public IDetailCustomization
+{
+public:
+
+	static TSharedRef<IDetailCustomization> MakeInstance()
+	{
+		return MakeShareable(new FSkinWeightDetailCustomization);
+	}
+
+	/** IDetailCustomization interface */
+	virtual void CustomizeDetails(IDetailLayoutBuilder& DetailBuilder) override;
+
+private:
+
+	void AddBrushUI(IDetailLayoutBuilder& DetailBuilder);
+	void AddSelectionUI(IDetailLayoutBuilder& DetailBuilder);
+
+	IDetailLayoutBuilder* CurrentDetailBuilder;
+	TWeakObjectPtr<USkinWeightsPaintToolProperties> SkinToolSettings;
+
+	static float WeightSliderWidths;
+	static float WeightEditingLabelsPercent;
+	static float WeightEditVerticalPadding;
+	static float WeightEditHorizontalPadding;
+};
+
+struct FWeightEditorElement
+{
+	int32 BoneIndex;
+
+	FWeightEditorElement(int32 InBoneIndex) : BoneIndex(InBoneIndex) {}
+};
+
+class SVertexWeightEditor;
+class SVertexWeightItem : public SMultiColumnTableRow<TSharedPtr<FWeightEditorElement>>
+{
+public:
+
+	SLATE_BEGIN_ARGS(SVertexWeightItem) {}
+	SLATE_ARGUMENT(TSharedPtr<FWeightEditorElement>, Element)
+	SLATE_ARGUMENT(TSharedPtr<SVertexWeightEditor>, ParentTable)
+	SLATE_END_ARGS()
+
+	virtual ~SVertexWeightItem() override {};
+
+	void Construct(const FArguments& InArgs, const TSharedRef<STableViewBase>& OwnerTableView);
+
+	virtual TSharedRef<SWidget> GenerateWidgetForColumn(const FName& ColumnName) override;
+
+private:
+	TSharedPtr<FWeightEditorElement> Element;
+	TSharedPtr<SVertexWeightEditor> ParentTable;
+	bool bInTransaction = false;
+};
+
+typedef SListView< TSharedPtr<FWeightEditorElement> > SWeightEditorListViewType;
+
+class SVertexWeightEditor : public SCompoundWidget
+{
+public:
+	SLATE_BEGIN_ARGS(SVertexWeightEditor) {}
+	SLATE_END_ARGS()
+
+	void Construct(const FArguments& InArgs, USkinWeightsPaintTool* InSkinTool);
+
+	virtual ~SVertexWeightEditor() override;
+	
+	void RefreshView();
+
+private:
+
+	// the vertex weight list view
+	TSharedPtr<SWeightEditorListViewType> ListView;
+	TArray< TSharedPtr<FWeightEditorElement> > ListViewItems;
+	TWeakObjectPtr<USkinWeightsPaintTool> Tool;
+	TArray<int32> SelectedVertices;
+
+	friend SVertexWeightItem;
+};
\ No newline at end of file
diff --git a/Engine/Plugins/Experimental/Animation/SkeletalMeshModelingTools/Source/SkeletalMeshModelingTools/Private/SkeletalMeshModelingToolsModule.cpp b/Engine/Plugins/Experimental/Animation/SkeletalMeshModelingTools/Source/SkeletalMeshModelingTools/Private/SkeletalMeshModelingToolsModule.cpp
--- a/Engine/Plugins/Experimental/Animation/SkeletalMeshModelingTools/Source/SkeletalMeshModelingTools/Private/SkeletalMeshModelingToolsModule.cpp
+++ b/Engine/Plugins/Experimental/Animation/SkeletalMeshModelingTools/Source/SkeletalMeshModelingTools/Private/SkeletalMeshModelingToolsModule.cpp
@@ -19,8 +19,10 @@
 #include "SkeletalMeshToolMenuContext.h"
 #include "ToolMenus.h"
 #include "DetailCustomization/SkeletonEditingToolPropertyCustomizations.h"
+#include "DetailCustomization/SkinWeightDetailCustomization.h"
 #include "Misc/ConfigCacheIni.h"
 #include "SkeletalMesh/SkeletonEditingTool.h"
+#include "SkeletalMesh/SkinWeightsPaintTool.h"
 #include "Styling/SlateIconFinder.h"
 #include "WorkflowOrientedApp/ApplicationMode.h"
 
@@ -225,6 +227,8 @@ void FSkeletalMeshModelingToolsModule::RegisterPropertyCustomizations()
 		CustomizedClasses.Add(InStructName);
 	};
 
+	
+	RegisterDetailCustomization(USkinWeightsPaintToolProperties::StaticClass()->GetFName(), &FSkinWeightDetailCustomization::MakeInstance);
 	RegisterDetailCustomization(USkeletonEditingTool::StaticClass()->GetFName(), &FSkeletonEditingToolDetailCustomization::MakeInstance);
 	RegisterDetailCustomization(USkeletonEditingProperties::StaticClass()->GetFName(), &FSkeletonEditingPropertiesDetailCustomization::MakeInstance);
 	RegisterDetailCustomization(UMirroringProperties::StaticClass()->GetFName(), &FMirroringPropertiesDetailCustomization::MakeInstance);
diff --git a/Engine/Plugins/Experimental/MeshModelingToolsetExp/Source/MeshModelingToolsEditorOnlyExp/Private/SkeletalMesh/SkinWeightsPaintTool.cpp b/Engine/Plugins/Experimental/MeshModelingToolsetExp/Source/MeshModelingToolsEditorOnlyExp/Private/SkeletalMesh/SkinWeightsPaintTool.cpp
--- a/Engine/Plugins/Experimental/MeshModelingToolsetExp/Source/MeshModelingToolsEditorOnlyExp/Private/SkeletalMesh/SkinWeightsPaintTool.cpp
+++ b/Engine/Plugins/Experimental/MeshModelingToolsetExp/Source/MeshModelingToolsEditorOnlyExp/Private/SkeletalMesh/SkinWeightsPaintTool.cpp
@@ -2192,9 +2192,19 @@ void USkinWeightsPaintTool::ToggleEditingMode()
 	SetFocusInViewport();
 }
 
-TObjectPtr<UPolygonSelectionMechanic> USkinWeightsPaintTool::GetSelectionMechanic()
+void USkinWeightsPaintTool::GrowSelection() const
 {
-	return PolygonSelectionMechanic;
+	PolygonSelectionMechanic->GrowSelection();
+}
+
+void USkinWeightsPaintTool::ShrinkSelection() const
+{
+	PolygonSelectionMechanic->ShrinkSelection();
+}
+
+void USkinWeightsPaintTool::FloodSelection() const
+{
+	PolygonSelectionMechanic->FloodSelection();
 }
 
 void USkinWeightsPaintTool::GetSelectedVertices(TArray<int32>& OutVertexIndices) const
diff --git a/Engine/Plugins/Experimental/MeshModelingToolsetExp/Source/MeshModelingToolsEditorOnlyExp/Public/SkeletalMesh/SkinWeightsPaintTool.h b/Engine/Plugins/Experimental/MeshModelingToolsetExp/Source/MeshModelingToolsEditorOnlyExp/Public/SkeletalMesh/SkinWeightsPaintTool.h
--- a/Engine/Plugins/Experimental/MeshModelingToolsetExp/Source/MeshModelingToolsEditorOnlyExp/Public/SkeletalMesh/SkinWeightsPaintTool.h
+++ b/Engine/Plugins/Experimental/MeshModelingToolsetExp/Source/MeshModelingToolsEditorOnlyExp/Public/SkeletalMesh/SkinWeightsPaintTool.h
@@ -384,7 +384,9 @@ public:
 	void ToggleEditingMode();
 
 	// edit selection
-	TObjectPtr<UPolygonSelectionMechanic> GetSelectionMechanic();
+	void GrowSelection() const;
+	void ShrinkSelection() const;
+	void FloodSelection() const;
 
 	// get a list of currently selected vertices
 	void GetSelectedVertices(TArray<int32>& OutVertexIndices) const;
diff --git a/Engine/Plugins/Experimental/MeshModelingToolsetExp/Source/ModelingEditorUI/Private/ModelingEditorUIModule.cpp b/Engine/Plugins/Experimental/MeshModelingToolsetExp/Source/ModelingEditorUI/Private/ModelingEditorUIModule.cpp
--- a/Engine/Plugins/Experimental/MeshModelingToolsetExp/Source/ModelingEditorUI/Private/ModelingEditorUIModule.cpp
+++ b/Engine/Plugins/Experimental/MeshModelingToolsetExp/Source/ModelingEditorUI/Private/ModelingEditorUIModule.cpp
@@ -2,41 +2,16 @@
 
 #include "ModelingEditorUIModule.h"
 
-#include "SkinWeightDetailCustomization.h"
-#include "SkeletalMesh/SkinWeightsPaintTool.h"
 #include "PropertyEditorModule.h"
 
 #define LOCTEXT_NAMESPACE "FModelingEditorUIModule"
 
 void FModelingEditorUIModule::StartupModule()
 {
-	FPropertyEditorModule& PropertyEditorModule = FModuleManager::LoadModuleChecked<FPropertyEditorModule>("PropertyEditor");
-	
-	auto RegisterDetailCustomization = [&](FName InStructName, auto InCustomizationFactory)
-    {
-    	PropertyEditorModule.RegisterCustomClassLayout(
-    		InStructName,
-    		FOnGetDetailCustomizationInstance::CreateStatic(InCustomizationFactory)
-    	);
-    	CustomizedClasses.Add(InStructName);
-    };
-
-	// register detail customizations
-	RegisterDetailCustomization(USkinWeightsPaintToolProperties::StaticClass()->GetFName(), &FSkinWeightDetailCustomization::MakeInstance);
-	
-	PropertyEditorModule.NotifyCustomizationModuleChanged();
 }
 
 void FModelingEditorUIModule::ShutdownModule()
 {
-	if (FPropertyEditorModule* PropertyModule = FModuleManager::GetModulePtr<FPropertyEditorModule>("PropertyEditor"))
-	{
-		for (const FName& ClassName : CustomizedClasses)
-		{
-			PropertyModule->UnregisterCustomClassLayout(ClassName);
-		}
-		PropertyModule->NotifyCustomizationModuleChanged();
-	}
 }
 
 #undef LOCTEXT_NAMESPACE
diff --git a/Engine/Plugins/Experimental/MeshModelingToolsetExp/Source/ModelingEditorUI/Private/SkinWeightDetailCustomization.cpp b/Engine/Plugins/Experimental/MeshModelingToolsetExp/Source/ModelingEditorUI/Private/SkinWeightDetailCustomization.cpp
deleted file mode 100644
diff --git a/Engine/Plugins/Experimental/MeshModelingToolsetExp/Source/ModelingEditorUI/Public/SkinWeightDetailCustomization.h b/Engine/Plugins/Experimental/MeshModelingToolsetExp/Source/ModelingEditorUI/Public/SkinWeightDetailCustomization.h
deleted file mode 100644
