Added ability to select edges and faces in skin paint tool.
Adds Grow/Shrink/Flood ability to both edges and faces in UPolygonSelectionMechanic

diff --git a/Engine/Plugins/Experimental/Animation/SkeletalMeshModelingTools/Source/SkeletalMeshModelingTools/Private/DetailCustomization/SkinWeightDetailCustomization.cpp b/Engine/Plugins/Experimental/Animation/SkeletalMeshModelingTools/Source/SkeletalMeshModelingTools/Private/DetailCustomization/SkinWeightDetailCustomization.cpp
--- a/Engine/Plugins/Experimental/Animation/SkeletalMeshModelingTools/Source/SkeletalMeshModelingTools/Private/DetailCustomization/SkinWeightDetailCustomization.cpp
+++ b/Engine/Plugins/Experimental/Animation/SkeletalMeshModelingTools/Source/SkeletalMeshModelingTools/Private/DetailCustomization/SkinWeightDetailCustomization.cpp
@@ -5,6 +5,7 @@
 #include "DetailCategoryBuilder.h"
 #include "DetailLayoutBuilder.h"
 #include "DetailWidgetRow.h"
+#include "ModelingToolsEditorModeStyle.h"
 #include "Widgets/Input/SSegmentedControl.h"
 #include "SkeletalMesh/SkinWeightsPaintTool.h"
 #include "SSkinWeightProfileImportOptions.h"
@@ -14,6 +15,7 @@
 #include "Widgets/Input/SButton.h"
 #include "Widgets/Input/SNumericEntryBox.h"
 #include "Widgets/Input/SSlider.h"
+#include "Framework/MultiBox/MultiBoxBuilder.h"
 
 #define LOCTEXT_NAMESPACE "SkinWeightToolSettingsEditor"
 
@@ -46,8 +48,8 @@
 		[
 			SNew(SSegmentedControl<EWeightEditMode>)
 			.ToolTipText(LOCTEXT("EditingModeTooltip",
-					"Brush: edit weights by painting directly on mesh.\n"
-					"Vertices: select vertices and edit weights directly.\n"
+					"Brush: edit weights by painting on mesh.\n"
+					"Mesh: select vertices/edges/faces to edit weights directly.\n"
 					"Bones: select and manipulate bones to preview deformations.\n"))
 			.Value_Lambda([this]()
 			{
@@ -64,8 +66,8 @@
 			})
 			+SSegmentedControl<EWeightEditMode>::Slot(EWeightEditMode::Brush)
 			.Text(LOCTEXT("BrushEditMode", "Brush"))
-			+ SSegmentedControl<EWeightEditMode>::Slot(EWeightEditMode::Vertices)
-			.Text(LOCTEXT("VertexEditMode", "Vertices"))
+			+ SSegmentedControl<EWeightEditMode>::Slot(EWeightEditMode::Mesh)
+			.Text(LOCTEXT("MeshEditMode", "Mesh"))
 			+ SSegmentedControl<EWeightEditMode>::Slot(EWeightEditMode::Bones)
 			.Text(LOCTEXT("BoneEditMode", "Bones"))
 		]
@@ -77,8 +79,8 @@
 		AddBrushUI(DetailBuilder);
 	}
 
-	// VERTEX editing mode UI
-	if (SkinToolSettings->EditingMode == EWeightEditMode::Vertices)
+	// MESH editing mode UI
+	if (SkinToolSettings->EditingMode == EWeightEditMode::Mesh)
 	{
 		AddSelectionUI(DetailBuilder);
 	}
@@ -306,6 +308,56 @@
 	IDetailCategoryBuilder& EditSelectionCategory = DetailBuilder.EditCategory("Edit Selection", FText::GetEmpty(), ECategoryPriority::Important);
 	EditSelectionCategory.InitiallyCollapsed(true);
 
+	// create a toolbar for the selection filter
+	FSlimHorizontalToolBarBuilder ToolbarBuilder(MakeShared<FUICommandList>(), FMultiBoxCustomization::None);
+	ToolbarBuilder.SetStyle(FModelingToolsEditorModeStyle::Get().Get(), "PolyEd.SelectionToolbar");
+	ToolbarBuilder.SetLabelVisibility(EVisibility::Collapsed);
+
+	ToolbarBuilder.BeginSection("SelectionFilter");
+	ToolbarBuilder.BeginBlockGroup();
+
+	TObjectPtr<USkinWeightsPaintToolProperties> ToolProperties = SkinToolSettings.Get();
+	auto AddToggleButtonForBool = [&ToolbarBuilder, ToolProperties](EComponentSelectionMode Mode, const FText& Label, const FText& Tooltip, const FName IconName)
+	{
+		ToolbarBuilder.AddToolBarButton(FUIAction(
+		FExecuteAction::CreateLambda([ToolProperties, Mode]()
+		{
+			ToolProperties->SetComponentMode(Mode);
+		}),
+		FCanExecuteAction::CreateLambda([ToolProperties]()
+		{
+			return ToolProperties->EditingMode == EWeightEditMode::Mesh;
+		}),
+		FIsActionChecked::CreateLambda([ToolProperties, Mode]()
+		{
+			return ToolProperties->ComponentSelectionMode == Mode;
+		})),
+		NAME_None,	// Extension hook
+		Label,		// Label
+		Tooltip,	// Tooltip
+		FSlateIcon(FModelingToolsEditorModeStyle::Get()->GetStyleSetName(), IconName),
+		EUserInterfaceActionType::ToggleButton);
+	};
+
+	AddToggleButtonForBool(
+		EComponentSelectionMode::Vertices,
+		LOCTEXT("VerticesLabel", "Vertices"),
+		LOCTEXT("VerticesTooltip", "Select mesh vertices."),
+		"PolyEd.SelectCorners");
+	AddToggleButtonForBool(
+		EComponentSelectionMode::Edges,
+		LOCTEXT("EdgesLabel", "Edges"),
+		LOCTEXT("EdgesTooltip", "Select mesh edges."),
+		"PolyEd.SelectEdges");
+	AddToggleButtonForBool(
+		EComponentSelectionMode::Faces,
+		LOCTEXT("FacesLabel", "Faces"),
+		LOCTEXT("FacesTooltip", "Select mesh faces."),
+		"PolyEd.SelectFaces");
+
+	ToolbarBuilder.EndBlockGroup();
+	ToolbarBuilder.EndSection();
+
 	// GROW/SHRINK/FLOOD Selection category
 	EditSelectionCategory.AddCustomRow(LOCTEXT("EditSelectionRow", "Edit Selection"), false)
 	.WholeRowContent()
@@ -313,5 +365,15 @@
 		SNew(SHorizontalBox)
 
 		+SHorizontalBox::Slot()
+		[
+			SNew(SHorizontalBox)
+			+SHorizontalBox::Slot()
+			.HAlign(HAlign_Left)
+			[
+				ToolbarBuilder.MakeWidget()
+			]
+		]
+		
+		+SHorizontalBox::Slot()
 		.Padding(2.f, WeightEditVerticalPadding)
 		[
@@ -336,7 +398,7 @@
 			.VAlign(VAlign_Center)
 			.Text(LOCTEXT("ShrinkSelectionButtonLabel", "Shrink"))
 			.ToolTipText(LOCTEXT("ShrinkSelectionTooltip",
-					"Shrink the current selection by removing vertices on the border of the current selection.\n"))
+					"Shrink the current selection by removing components on the border of the current selection.\n"))
 			.OnClicked_Lambda([this]()
 			{
 				SkinToolSettings->WeightTool->ShrinkSelection();
@@ -352,7 +414,7 @@
 			.VAlign(VAlign_Center)
 			.Text(LOCTEXT("FloodSelectionButtonLabel", "Flood"))
 			.ToolTipText(LOCTEXT("FloodSelectionTooltip",
-					"Flood the current selection by adding all connected vertices to the current selection.\n"))
+					"Flood the current selection by adding all connected components to the current selection.\n"))
 			.OnClicked_Lambda([this]()
 			{
 				SkinToolSettings->WeightTool->FloodSelection();
@@ -365,144 +427,6 @@
 	IDetailCategoryBuilder& EditWeightsCategory = DetailBuilder.EditCategory("Edit Weights", FText::GetEmpty(), ECategoryPriority::Important);
 	EditWeightsCategory.InitiallyCollapsed(true);
 
-	// AVERAGE/RELAX/NORMALIZE WEIGHTS category
-	EditWeightsCategory.AddCustomRow(LOCTEXT("NormalizeWeightsRow", "Normalize"), false)
-	.WholeRowContent()
-	[
-		SNew(SHorizontalBox)
-
-		+SHorizontalBox::Slot()
-		.Padding(2.f, WeightEditVerticalPadding)
-		[
-			SNew(SButton)
-			.HAlign(HAlign_Center)
-			.VAlign(VAlign_Center)
-			.Text(LOCTEXT("AverageWeightsButtonLabel", "Average"))
-			.ToolTipText(LOCTEXT("AverageWeightsTooltip",
-					"Takes the average of vertex weights and applies it to the selected vertices.\n"
-					"This command operates on the selected bone(s) and selected vertices.\n "
-					"If no bones are selected, ALL bone weights are considered.\n "
-					"If no vertices are selected, ALL vertices are considered."))
-			.OnClicked_Lambda([this]()
-			{
-				SkinToolSettings->WeightTool->AverageWeights();
-				return FReply::Handled();
-			})
-		]
-
-		+SHorizontalBox::Slot()
-		.Padding(2.f, WeightEditVerticalPadding)
-		[
-			SNew(SButton)
-			.HAlign(HAlign_Center)
-			.VAlign(VAlign_Center)
-			.Text(LOCTEXT("NormalizeWeightsButtonLabel", "Normalize"))
-			.ToolTipText(LOCTEXT("NormalizeWeightsTooltip",
-					"Forces the weights on the selected vertices to sum to 1.\n"
-					"This command operates on the selected vertices.\n "
-					"If no vertices are selected, ALL vertices are considered."))
-			.IsEnabled_Lambda([this]()
-			{
-				return SkinToolSettings->EditingMode == EWeightEditMode::Vertices;
-			})
-			.OnClicked_Lambda([this]()
-			{
-				SkinToolSettings->WeightTool->NormalizeWeights();
-				return FReply::Handled();
-			})
-		]
-	];
-
-	// MIRROR WEIGHTS category
-	EditWeightsCategory.AddCustomRow(LOCTEXT("MirrorWeightsRow", "Mirror"), false)
-	.WholeRowContent()
-	[
-		SNew(SVerticalBox)
-		
-		+ SVerticalBox::Slot()
-		.Padding(WeightEditHorizontalPadding, WeightEditVerticalPadding)
-		[
-			SNew(SHorizontalBox)
-			
-			+SHorizontalBox::Slot()
-			.VAlign(VAlign_Center)
-			.FillWidth(WeightEditingLabelsPercent)
-			[
-				SNew(STextBlock)
-				.Text(LOCTEXT("MirrorPlaneLabel", "Mirror Plane"))
-				.Font(FAppStyle::GetFontStyle(TEXT("PropertyWindow.NormalFont")))
-				.ToolTipText(LOCTEXT("MirrorPlaneTooltip", "The plane to copy weights across."))
-			]
-
-			+SHorizontalBox::Slot()
-			.FillWidth(1.f)
-			[
-				SNew(SHorizontalBox)
-				+SHorizontalBox::Slot()
-				[
-					SNew(SSegmentedControl<EAxis::Type>)
-					.ToolTipText(LOCTEXT("MirrorAxisTooltip",
-						"X: copies weights across the YZ plane.\n"
-						"Y: copies weights across the XZ plane.\n"
-						"Z: copies weights across the XY plane."))
-					.Value_Lambda([this]()
-					{
-						return SkinToolSettings->MirrorAxis;
-					})
-					.OnValueChanged_Lambda([this](EAxis::Type Mode)
-					{
-						SkinToolSettings->MirrorAxis = Mode;
-					})
-					+ SSegmentedControl<EAxis::Type>::Slot(EAxis::X)
-					.Text(LOCTEXT("MirrorXLabel", "X"))
-					+ SSegmentedControl<EAxis::Type>::Slot(EAxis::Y)
-					.Text(LOCTEXT("MirrorYLabel", "Y"))
-					+ SSegmentedControl<EAxis::Type>::Slot(EAxis::Z)
-					.Text(LOCTEXT("MirrorZLabel", "Z"))
-				]
-					
-				+SHorizontalBox::Slot()
-				[
-					SNew(SSegmentedControl<EMirrorDirection>)
-					.ToolTipText(LOCTEXT("MirrorDirectionTooltip", "The direction that determines what side of the plane to copy weights from."))
-					.Value_Lambda([this]()
-					{
-						return SkinToolSettings->MirrorDirection;
-					})
-					.OnValueChanged_Lambda([this](EMirrorDirection Mode)
-					{
-						SkinToolSettings->MirrorDirection = Mode;
-					})
-					+ SSegmentedControl<EMirrorDirection>::Slot(EMirrorDirection::PositiveToNegative)
-					.Text(LOCTEXT("MirrorPosToNegLabel", "+ to -"))
-					+ SSegmentedControl<EMirrorDirection>::Slot(EMirrorDirection::NegativeToPositive)
-					.Text(LOCTEXT("MirrorNegToPosLabel", "- to +"))
-				]
-			]
-		]
-
-		+SVerticalBox::Slot()
-		.Padding(0.f, WeightEditVerticalPadding)
-		[
-			SNew(SBox)
-			[
-				SNew(SButton)
-				.HAlign(HAlign_Center)
-				.Text(LOCTEXT("MirrorWeightsButtonLabel", "Mirror"))
-				.ToolTipText(LOCTEXT("MirrorButtonTooltip",
-					"Weights are copied across the given plane in the given direction.\n"
-					"This command operates on the selected bone(s) and selected vertices.\n "
-					"If no bones are selected, ALL bone weights are considered.\n "
-					"If no vertices are selected, ALL vertices are considered."))
-				.OnClicked_Lambda([this]()
-				{
-					SkinToolSettings->WeightTool->MirrorWeights(SkinToolSettings->MirrorAxis, SkinToolSettings->MirrorDirection);
-					return FReply::Handled();
-				})
-			]
-		]
-	];
-
 	// FLOOD WEIGHTS category
 	EditWeightsCategory.AddCustomRow(LOCTEXT("FloodWeightsRow", "Flood"), false)
 	.WholeRowContent()
@@ -560,7 +484,7 @@
 				SNew(STextBlock)
 				.Text(LOCTEXT("FloodOperationLabel", "Flood Operation"))
 				.Font(FAppStyle::GetFontStyle(TEXT("PropertyWindow.NormalFont")))
-				.ToolTipText(LOCTEXT("FloodOperationTooltip", "The various flood weight operations."))
+				.ToolTipText(LOCTEXT("FloodOperationTooltip", "Various weight editing operations applied to the current selection."))
 			]
 
 			+SHorizontalBox::Slot()
@@ -571,7 +495,7 @@
 					SNew(SButton)
 					.HAlign(HAlign_Center)
 					.Text(LOCTEXT("AddWeightsButtonLabel", "Add"))
-					.ToolTipText(LOCTEXT("AddOpTooltip", "Add: applies the current weight plus the flood amount to the new weight."))
+					.ToolTipText(LOCTEXT("AddOpTooltip", "Add: adds the Flood Amount value to the current weight."))
 					.OnClicked_Lambda([this]()
 					{
 						SkinToolSettings->WeightTool->FloodWeights(SkinToolSettings->FloodValue, EWeightEditOperation::Add);
@@ -618,6 +542,54 @@
 		]
 	];
 
+	// AVERAGE/NORMALIZE WEIGHTS category
+	EditWeightsCategory.AddCustomRow(LOCTEXT("NormalizeWeightsRow", "Normalize"), false)
+	.WholeRowContent()
+	[
+		SNew(SHorizontalBox)
+
+		+SHorizontalBox::Slot()
+		.Padding(2.f, WeightEditVerticalPadding)
+		[
+			SNew(SButton)
+			.HAlign(HAlign_Center)
+			.VAlign(VAlign_Center)
+			.Text(LOCTEXT("AverageWeightsButtonLabel", "Average"))
+			.ToolTipText(LOCTEXT("AverageWeightsTooltip",
+					"Takes the average of vertex weights and applies it to the selected vertices.\n"
+					"This command operates on the selected bone(s) and selected vertices.\n "
+					"If no bones are selected, ALL bone weights are considered.\n "
+					"If no vertices are selected, ALL vertices are considered."))
+			.OnClicked_Lambda([this]()
+			{
+				SkinToolSettings->WeightTool->AverageWeights();
+				return FReply::Handled();
+			})
+		]
+
+		+SHorizontalBox::Slot()
+		.Padding(2.f, WeightEditVerticalPadding)
+		[
+			SNew(SButton)
+			.HAlign(HAlign_Center)
+			.VAlign(VAlign_Center)
+			.Text(LOCTEXT("NormalizeWeightsButtonLabel", "Normalize"))
+			.ToolTipText(LOCTEXT("NormalizeWeightsTooltip",
+					"Forces the weights on the selected vertices to sum to 1.\n"
+					"This command operates on the selected vertices.\n "
+					"If no vertices are selected, ALL vertices are considered."))
+			.IsEnabled_Lambda([this]()
+			{
+				return SkinToolSettings->EditingMode == EWeightEditMode::Mesh;
+			})
+			.OnClicked_Lambda([this]()
+			{
+				SkinToolSettings->WeightTool->NormalizeWeights();
+				return FReply::Handled();
+			})
+		]
+	];
+	
 	// PRUNE WEIGHTS category
 	EditWeightsCategory.AddCustomRow(LOCTEXT("PruneWeightsRow", "Prune"), false)
 	.WholeRowContent()
@@ -634,13 +606,12 @@
 			.FillWidth(WeightEditingLabelsPercent)
 			[
 				SNew(STextBlock)
-				.Text(LOCTEXT("PruneThresholdLabel", "Prune Threshold"))
+				.Text(LOCTEXT("PruneThresholdLabel", "Prune"))
 				.Font(FAppStyle::GetFontStyle(TEXT("PropertyWindow.NormalFont")))
 				.ToolTipText(LOCTEXT("PruneThresholdTooltip", "The threshold weight value to use with Prune operation."))
 			]
 
 			+SHorizontalBox::Slot()
-			.MaxWidth(WeightSliderWidths)
 			[
 				SNew(SSpinBox<float>)
 				.MinValue(0.f)
@@ -658,25 +629,113 @@
 					SkinToolSettings->SaveConfig();
 				})
 			]
+
+			+SHorizontalBox::Slot()
+			[
+				SNew(SBox)
+				[
+					SNew(SButton)
+					.HAlign(HAlign_Center)
+					.Text(LOCTEXT("PruneWeightsButtonLabel", "Prune"))
+					.ToolTipText(LOCTEXT("PruneButtonTooltip",
+						"Weights below the given threshold value are removed.\n"
+						"This command operates on the selected bone(s) and selected vertices.\n "
+						"If no bones are selected, ALL bone weights are considered.\n "
+						"If no vertices are selected, ALL vertices are considered."))
+					.OnClicked_Lambda([this]()
+					{
+						SkinToolSettings->WeightTool->PruneWeights(SkinToolSettings->PruneValue);
+						return FReply::Handled();
+					})
+				]
+			]
 		]
+	];
+
+	// MIRROR WEIGHTS category
+	EditWeightsCategory.AddCustomRow(LOCTEXT("MirrorWeightsRow", "Mirror"), false)
+	.WholeRowContent()
+	[
+		SNew(SVerticalBox)
 		
-		+SVerticalBox::Slot()
+		+ SVerticalBox::Slot()
 		.Padding(WeightEditHorizontalPadding, WeightEditVerticalPadding)
 		[
+			SNew(SHorizontalBox)
+			
+			+SHorizontalBox::Slot()
+			.VAlign(VAlign_Center)
+			.FillWidth(WeightEditingLabelsPercent)
+			[
+				SNew(STextBlock)
+				.Text(LOCTEXT("MirrorPlaneLabel", "Mirror Plane"))
+				.Font(FAppStyle::GetFontStyle(TEXT("PropertyWindow.NormalFont")))
+				.ToolTipText(LOCTEXT("MirrorPlaneTooltip", "The plane to copy weights across."))
+			]
+
+			+SHorizontalBox::Slot()
+			.FillWidth(1.f)
+			[
+				SNew(SHorizontalBox)
+				+SHorizontalBox::Slot()
+				[
+					SNew(SSegmentedControl<EAxis::Type>)
+					.ToolTipText(LOCTEXT("MirrorAxisTooltip",
+						"X: copies weights across the YZ plane.\n"
+						"Y: copies weights across the XZ plane.\n"
+						"Z: copies weights across the XY plane."))
+					.Value_Lambda([this]()
+					{
+						return SkinToolSettings->MirrorAxis;
+					})
+					.OnValueChanged_Lambda([this](EAxis::Type Mode)
+					{
+						SkinToolSettings->MirrorAxis = Mode;
+					})
+					+ SSegmentedControl<EAxis::Type>::Slot(EAxis::X)
+					.Text(LOCTEXT("MirrorXLabel", "X"))
+					+ SSegmentedControl<EAxis::Type>::Slot(EAxis::Y)
+					.Text(LOCTEXT("MirrorYLabel", "Y"))
+					+ SSegmentedControl<EAxis::Type>::Slot(EAxis::Z)
+					.Text(LOCTEXT("MirrorZLabel", "Z"))
+				]
+					
+				+SHorizontalBox::Slot()
+				[
+					SNew(SSegmentedControl<EMirrorDirection>)
+					.ToolTipText(LOCTEXT("MirrorDirectionTooltip", "The direction that determines what side of the plane to copy weights from."))
+					.Value_Lambda([this]()
+					{
+						return SkinToolSettings->MirrorDirection;
+					})
+					.OnValueChanged_Lambda([this](EMirrorDirection Mode)
+					{
+						SkinToolSettings->MirrorDirection = Mode;
+					})
+					+ SSegmentedControl<EMirrorDirection>::Slot(EMirrorDirection::PositiveToNegative)
+					.Text(LOCTEXT("MirrorPosToNegLabel", "+ to -"))
+					+ SSegmentedControl<EMirrorDirection>::Slot(EMirrorDirection::NegativeToPositive)
+					.Text(LOCTEXT("MirrorNegToPosLabel", "- to +"))
+				]
+			]
+		]
+
+		+SVerticalBox::Slot()
+		.Padding(0.f, WeightEditVerticalPadding)
+		[
 			SNew(SBox)
-			.MinDesiredWidth(WeightSliderWidths)
 			[
 				SNew(SButton)
 				.HAlign(HAlign_Center)
-				.Text(LOCTEXT("PruneWeightsButtonLabel", "Prune"))
-				.ToolTipText(LOCTEXT("PruneButtonTooltip",
-					"Weights below the given threshold value are removed.\n"
+				.Text(LOCTEXT("MirrorWeightsButtonLabel", "Mirror"))
+				.ToolTipText(LOCTEXT("MirrorButtonTooltip",
+					"Weights are copied across the given plane in the given direction.\n"
 					"This command operates on the selected bone(s) and selected vertices.\n "
 					"If no bones are selected, ALL bone weights are considered.\n "
 					"If no vertices are selected, ALL vertices are considered."))
 				.OnClicked_Lambda([this]()
 				{
-					SkinToolSettings->WeightTool->PruneWeights(SkinToolSettings->PruneValue);
+					SkinToolSettings->WeightTool->MirrorWeights(SkinToolSettings->MirrorAxis, SkinToolSettings->MirrorDirection);
 					return FReply::Handled();
 				})
 			]
@@ -684,7 +743,7 @@
 	];
 
 	// VERTEX EDITOR category
-	EditWeightsCategory.AddCustomRow(LOCTEXT("VertexEditorRow", "Vertex Editor"), false)
+	EditWeightsCategory.AddCustomRow(LOCTEXT("VertexEditorRow", "Component Editor"), false)
 	.WholeRowContent()
 	[
 		SNew(SVertexWeightEditor, SkinToolSettings->WeightTool)
diff --git a/Engine/Plugins/Experimental/MeshModelingToolsetExp/Source/MeshModelingToolsEditorOnlyExp/Private/SkeletalMesh/SkinWeightsPaintTool.cpp b/Engine/Plugins/Experimental/MeshModelingToolsetExp/Source/MeshModelingToolsEditorOnlyExp/Private/SkeletalMesh/SkinWeightsPaintTool.cpp
--- a/Engine/Plugins/Experimental/MeshModelingToolsetExp/Source/MeshModelingToolsEditorOnlyExp/Private/SkeletalMesh/SkinWeightsPaintTool.cpp
+++ b/Engine/Plugins/Experimental/MeshModelingToolsetExp/Source/MeshModelingToolsEditorOnlyExp/Private/SkeletalMesh/SkinWeightsPaintTool.cpp
@@ -12,6 +12,8 @@
 #include "TargetInterfaces/PrimitiveComponentBackedTarget.h"
 #include "ModelingToolTargetUtil.h"
 
+#include "Materials/Material.h"
+#include "Materials/MaterialInstanceDynamic.h"
 #include "MeshDescription.h"
 #include "MeshModelingToolsEditorOnlyExp.h"
 #include "PointSetAdapter.h"
@@ -301,7 +303,7 @@ void FSkinToolDeformer::UpdateVertexDeformation(USkinWeightsPaintTool* Tool)
 	}
 
 	// update data structures used by the selection mode
-	if (EditingMode == EWeightEditMode::Vertices)
+	if (EditingMode == EWeightEditMode::Mesh)
 	{
 		// update AABB Tree for vertex selection
 		TRACE_CPUPROFILER_EVENT_SCOPE(SkinTool::UpdateAABBTree);
@@ -791,25 +793,46 @@ void USkinWeightsPaintTool::Setup()
 
 	RecalculateBrushRadius();
 
+	
+	const FLinearColor FaceSelectedOrange = FLinearColor(0.886f, 0.672f, 0.473f);
+	const FLinearColor VertexSelectedPurple = FLinearColor(0.78f, 0.f, 0.78f);
+	const FLinearColor VertexSelectedYellow = FLinearColor::Yellow;
+
+	// configure secondary render material for selected triangles
+	// NOTE: the material returned by ToolSetupUtil::GetSelectionMaterial has a checkerboard pattern on back faces which makes it hard to use
+	UMaterialInterface* Material = LoadObject<UMaterial>(nullptr, TEXT("/MeshModelingToolsetExp/Materials/SculptMaterial"));
+	if (Material)
+	{
+		if (UMaterialInstanceDynamic* MatInstance = UMaterialInstanceDynamic::Create(Material, GetToolManager()))
+		{
+			MatInstance->SetVectorParameterValue(TEXT("Color"), FaceSelectedOrange);
+			PreviewMesh->SetSecondaryRenderMaterial(MatInstance);
+		}
+	}
+
 	// set up vertex selection mechanic
 	PolygonSelectionMechanic = NewObject<UPolygonSelectionMechanic>(this);
 	PolygonSelectionMechanic->bAddSelectionFilterPropertiesToParentTool = false;
 	PolygonSelectionMechanic->Setup(this);
 	PolygonSelectionMechanic->SetIsEnabled(false);
 	PolygonSelectionMechanic->OnSelectionChanged.AddLambda([this](){OnSelectionChanged.Broadcast();} );
-	// only select vertices
-	PolygonSelectionMechanic->Properties->bSelectEdges = false;
-	PolygonSelectionMechanic->Properties->bSelectFaces = false;
-	PolygonSelectionMechanic->Properties->bSelectVertices = true;
+	// restore saved mode
+	SetComponentSelectionMode(WeightToolProperties->ComponentSelectionMode);
 	// adjust selection rendering for this context
-	PolygonSelectionMechanic->PolyEdgesRenderer.PointColor = FLinearColor(0.78f, 0.f, 0.78f);
-	PolygonSelectionMechanic->PolyEdgesRenderer.PointSize = 5.0f;
-	PolygonSelectionMechanic->HilightRenderer.PointColor = FLinearColor::Red;
+	PolygonSelectionMechanic->HilightRenderer.PointColor = FLinearColor::Blue;
 	PolygonSelectionMechanic->HilightRenderer.PointSize = 10.0f;
-	PolygonSelectionMechanic->SelectionRenderer.LineThickness = 0.0f;
-	PolygonSelectionMechanic->SelectionRenderer.PointColor = FLinearColor::Yellow;
+
+	PolygonSelectionMechanic->SelectionRenderer.LineThickness = 1.0f;
+	PolygonSelectionMechanic->SelectionRenderer.PointColor = VertexSelectedYellow;
 	PolygonSelectionMechanic->SelectionRenderer.PointSize = 5.0f;
-	PolygonSelectionMechanic->SetShowEdges(false);
+	PolygonSelectionMechanic->SelectionRenderer.DepthBias = 2.0f;
+
+	// despite the name, this renders the vertices
+	PolygonSelectionMechanic->PolyEdgesRenderer.PointColor = VertexSelectedPurple;
+	PolygonSelectionMechanic->PolyEdgesRenderer.PointSize = 5.0f;
+	PolygonSelectionMechanic->PolyEdgesRenderer.DepthBias = 2.0f;
+	PolygonSelectionMechanic->PolyEdgesRenderer.LineThickness = 1.0f;
+	
 	// initialize the polygon selection mechanic
 	constexpr bool bAutoBuild = true;
 	const FDynamicMesh3* DynamicMesh = PreviewMesh->GetPreviewDynamicMesh();
@@ -822,6 +845,20 @@ void USkinWeightsPaintTool::Setup()
 		SelectionTopology.Get(),
 		[this]() { return MeshSpatial.Get(); }
 	);
+	// secondary triangle buffer used to render face selection
+	PreviewMesh->EnableSecondaryTriangleBuffers([this](const FDynamicMesh3* Mesh, int32 TriangleID)
+	{
+		return PolygonSelectionMechanic->GetActiveSelection().IsSelectedTriangle(Mesh, SelectionTopology.Get(), TriangleID);
+	});
+	// notify preview mesh when triangle selection has been updated
+	PolygonSelectionMechanic->OnSelectionChanged.AddWeakLambda(this, [this]()
+	{
+		PreviewMesh->FastNotifySecondaryTrianglesChanged();
+	});
+	PolygonSelectionMechanic->OnFaceSelectionPreviewChanged.AddWeakLambda(this, [this]()
+	{
+		PreviewMesh->FastNotifySecondaryTrianglesChanged();
+	});
 	
 	SmoothWeightsDataSource = MakeUnique<FPaintToolWeightsDataSource>(&Weights);
 	SmoothWeightsOp = MakeUnique<UE::Geometry::TSmoothBoneWeights<int32, float>>(PreviewMesh->GetMesh(), SmoothWeightsDataSource.Get());
@@ -860,7 +897,7 @@ void USkinWeightsPaintTool::Render(IToolsContextRenderAPI* RenderAPI)
 	{
 		Super::Render(RenderAPI);	
 	}
-	else if (PolygonSelectionMechanic && WeightToolProperties->EditingMode == EWeightEditMode::Vertices)
+	else if (PolygonSelectionMechanic && WeightToolProperties->EditingMode == EWeightEditMode::Mesh)
 	{
 		PolygonSelectionMechanic->Render(RenderAPI);
 	}
@@ -877,26 +914,24 @@ FBox USkinWeightsPaintTool::GetWorldSpaceFocusBox()
 				return FBox(LastBrushStamp.WorldPosition - Radius, LastBrushStamp.WorldPosition + Radius);
 			}
 			break;
-	case EWeightEditMode::Vertices:
+	case EWeightEditMode::Mesh:
 		{
-			if (PolygonSelectionMechanic)
+			FAxisAlignedBox3d Bounds = FAxisAlignedBox3d::Empty();
+			TArray<int32> SelectedVertexIndices;
+			GetSelectedVertices(SelectedVertexIndices);
+			if (!SelectedVertexIndices.IsEmpty())
 			{
-				const FGroupTopologySelection& Selection = PolygonSelectionMechanic->GetActiveSelection();
-				if (!Selection.SelectedCornerIDs.IsEmpty())
+				const FDynamicMesh3* Mesh = PreviewMesh->GetMesh();
+				const FTransform3d Transform(PreviewMesh->GetTransform());
+				for (const int32 VertexID : SelectedVertexIndices)
 				{
-					const FDynamicMesh3* Mesh = PreviewMesh->GetMesh();
-					const FTransform3d Transform(PreviewMesh->GetTransform());
-					FAxisAlignedBox3d Bounds = FAxisAlignedBox3d::Empty();
-					for (const int32 VertexID : Selection.SelectedCornerIDs)
-					{
-						Bounds.Contain(Transform.TransformPosition(Mesh->GetVertex(VertexID)));
-					}
-					if (Bounds.MaxDim() > FMathf::ZeroTolerance)
-					{
-						return static_cast<FBox>(Bounds);
-					}
+					Bounds.Contain(Transform.TransformPosition(Mesh->GetVertex(VertexID)));
 				}
 			}
+			if (Bounds.MaxDim() > FMathf::ZeroTolerance)
+			{
+				return static_cast<FBox>(Bounds);
+			}
 		}
 		break;
 	case EWeightEditMode::Bones:
@@ -999,6 +1034,14 @@ void USkinWeightsPaintTool::OnTick(float DeltaTime)
 	Weights.Deformer.UpdateVertexDeformation(this);
 }
 
+void USkinWeightsPaintToolProperties::SetComponentMode(EComponentSelectionMode InComponentMode)
+{
+	ComponentSelectionMode = InComponentMode;
+	
+	WeightTool->SetComponentSelectionMode(ComponentSelectionMode);
+	WeightTool->SetFocusInViewport();
+}
+
 void USkinWeightsPaintToolProperties::SetFalloffMode(EWeightBrushFalloffMode InFalloffMode)
 {
 	bColorModeChanged = true;
@@ -1686,6 +1729,11 @@ void USkinWeightsPaintTool::OnShutdown(EToolShutdownType ShutdownType)
 	FPreviewProfileController PreviewProfileController;
 	PreviewProfileController.SetActiveProfile(PreviewProfileToRestore);
 	GetMutableDefault<UPersonaOptions>()->bShowBoneColors = bBoneColorsToRestore;
+	// mesh selection mode takes control of "Show Edges" render flag
+	if (IConsoleVariable* CVar = IConsoleManager::Get().FindConsoleVariable(TEXT("ShowFlag.MeshEdges")))
+	{
+		CVar->Unset(ECVF_SetByCode);
+	}
 
 	if (EditorContext.IsValid())
 	{
@@ -2170,8 +2218,8 @@ void USkinWeightsPaintTool::ToggleEditingMode()
 	// toggle brush mode
 	SetBrushEnabled(WeightToolProperties->EditingMode == EWeightEditMode::Brush);
 
-	// toggle vertex mode
-	PolygonSelectionMechanic->SetIsEnabled(WeightToolProperties->EditingMode == EWeightEditMode::Vertices);
+	// toggle mesh mode
+	PolygonSelectionMechanic->SetIsEnabled(WeightToolProperties->EditingMode == EWeightEditMode::Mesh);
 
 	// toggle bone select mode
 	// this mode is set to be compatible with the 
@@ -2190,6 +2238,23 @@ void USkinWeightsPaintTool::ToggleEditingMode()
 	SetFocusInViewport();
 }
 
+void USkinWeightsPaintTool::SetComponentSelectionMode(EComponentSelectionMode InMode)
+{
+	PolygonSelectionMechanic->Properties->bSelectVertices = InMode == EComponentSelectionMode::Vertices;
+	PolygonSelectionMechanic->Properties->bSelectEdges = InMode == EComponentSelectionMode::Edges;
+	PolygonSelectionMechanic->Properties->bSelectFaces = InMode == EComponentSelectionMode::Faces;
+	
+	PolygonSelectionMechanic->SetShowSelectableCorners(InMode == EComponentSelectionMode::Vertices);
+	PolygonSelectionMechanic->SetShowEdges(InMode == EComponentSelectionMode::Edges);
+
+	IConsoleVariable* CVar = IConsoleManager::Get().FindConsoleVariable(TEXT("ShowFlag.MeshEdges"));
+	if (CVar)
+	{
+		const float Value = (InMode == EComponentSelectionMode::Edges) ? 0.0f : 1.0f;
+		CVar->Set(Value, ECVF_SetByCode);
+	}
+}
+
 void USkinWeightsPaintTool::GrowSelection() const
 {
 	PolygonSelectionMechanic->GrowSelection();
@@ -2207,18 +2272,61 @@ void USkinWeightsPaintTool::FloodSelection() const
 
 void USkinWeightsPaintTool::GetSelectedVertices(TArray<int32>& OutVertexIndices) const
 {
+	OutVertexIndices.Empty();
+	if (!PolygonSelectionMechanic)
+	{
+		return;
+	}
+	
 	const FGroupTopologySelection& Selection = PolygonSelectionMechanic->GetActiveSelection();
-	if (!Selection.SelectedCornerIDs.IsEmpty())
+	const FDynamicMesh3* DynamicMesh = PreviewMesh->GetMesh();
+	const FVertexArray& AllVertices = EditedMesh->Vertices();
+
+	// validate and add vertices to the output array
+	auto AddVertices = [&OutVertexIndices, &AllVertices](const TSet<int32>& VerticesToAdd)
 	{
 		// we have to make sure that the vertex ids are safe to use as PolygonSelectionMechanic does not act on the
 		// mesh description but on the dynamic mesh that can duplicate vertices when dealing with degenerate triangles.
 		// cf. FMeshDescriptionToDynamicMesh::Convert for more details.
-		const FVertexArray& Vertices = EditedMesh->Vertices();
-		OutVertexIndices.Empty();
-		Algo::CopyIf(Selection.SelectedCornerIDs, OutVertexIndices, [&](int VertexID)
+		Algo::CopyIf(VerticesToAdd, OutVertexIndices, [&](int32 VertexID)
 		{
-			return Vertices.IsValid(VertexID);	
+			return AllVertices.IsValid(VertexID);	
 		});
+	};
+
+	// add selected vertices
+	if (PolygonSelectionMechanic->Properties->bSelectVertices)
+	{
+		AddVertices(Selection.SelectedCornerIDs);
+	}
+
+	// add vertices on selected edges
+	if (PolygonSelectionMechanic->Properties->bSelectEdges)
+	{
+		TSet<int32> VerticesInSelectedEdges;
+		for (const int32 SelectedEdgeIndex : Selection.SelectedEdgeIDs)
+		{
+			FDynamicMesh3::FEdge CurrentEdge = DynamicMesh->GetEdge(SelectedEdgeIndex);
+			VerticesInSelectedEdges.Add(CurrentEdge.Vert.A);
+			VerticesInSelectedEdges.Add(CurrentEdge.Vert.B);
+		}
+		
+		AddVertices(VerticesInSelectedEdges);
+	}
+
+	// add vertices in selected faces
+	if (PolygonSelectionMechanic->Properties->bSelectFaces)
+	{
+		TSet<int32> VerticesInSelectedFaces;
+		for (const int32 SelectedFaceIndex : Selection.SelectedGroupIDs)
+		{
+			UE::Geometry::FIndex3i TriangleVertices = DynamicMesh->GetTriangleRef(SelectedFaceIndex);
+			VerticesInSelectedFaces.Add(TriangleVertices[0]);
+			VerticesInSelectedFaces.Add(TriangleVertices[1]);
+			VerticesInSelectedFaces.Add(TriangleVertices[2]);
+		}
+		
+		AddVertices(VerticesInSelectedFaces);
 	}
 }
 
diff --git a/Engine/Plugins/Experimental/MeshModelingToolsetExp/Source/MeshModelingToolsEditorOnlyExp/Public/SkeletalMesh/SkinWeightsPaintTool.h b/Engine/Plugins/Experimental/MeshModelingToolsetExp/Source/MeshModelingToolsEditorOnlyExp/Public/SkeletalMesh/SkinWeightsPaintTool.h
--- a/Engine/Plugins/Experimental/MeshModelingToolsetExp/Source/MeshModelingToolsEditorOnlyExp/Public/SkeletalMesh/SkinWeightsPaintTool.h
+++ b/Engine/Plugins/Experimental/MeshModelingToolsetExp/Source/MeshModelingToolsEditorOnlyExp/Public/SkeletalMesh/SkinWeightsPaintTool.h
@@ -38,10 +38,19 @@ UENUM()
 enum class EWeightEditMode : uint8
 {
 	Brush,
-	Vertices,
+	Mesh,
 	Bones,
 };
 
+// component selection mode
+UENUM()
+enum class EComponentSelectionMode : uint8
+{
+	Vertices,
+	Edges,
+	Faces
+};
+
 // weight color mode
 UENUM()
 enum class EWeightColorMode : uint8
@@ -289,6 +298,10 @@ public:
 	EWeightEditOperation BrushMode;
 	EWeightEditOperation PriorBrushMode; // when toggling with modifier key
 
+	// are we selecting vertices, edges or faces
+	UPROPERTY(Config)
+	EComponentSelectionMode ComponentSelectionMode;
+
 	// weight color properties
 	UPROPERTY(EditAnywhere, Config, Category = MeshDisplay)
 	EWeightColorMode ColorMode;
@@ -325,6 +338,7 @@ public:
 	// pointer back to paint tool
 	TObjectPtr<USkinWeightsPaintTool> WeightTool;
 
+	void SetComponentMode(EComponentSelectionMode InComponentMode);
 	void SetFalloffMode(EWeightBrushFalloffMode InFalloffMode);
 	void SetColorMode(EWeightColorMode InColorMode);
 	void SetBrushMode(EWeightEditOperation InBrushMode);
@@ -384,6 +398,7 @@ public:
 	void ToggleEditingMode();
 
 	// edit selection
+	void SetComponentSelectionMode(EComponentSelectionMode InMode);
 	void GrowSelection() const;
 	void ShrinkSelection() const;
 	void FloodSelection() const;
diff --git a/Engine/Plugins/Runtime/MeshModelingToolset/Source/ModelingComponents/Private/Selection/MeshTopologySelectionMechanic.cpp b/Engine/Plugins/Runtime/MeshModelingToolset/Source/ModelingComponents/Private/Selection/MeshTopologySelectionMechanic.cpp
--- a/Engine/Plugins/Runtime/MeshModelingToolset/Source/ModelingComponents/Private/Selection/MeshTopologySelectionMechanic.cpp
+++ b/Engine/Plugins/Runtime/MeshModelingToolset/Source/ModelingComponents/Private/Selection/MeshTopologySelectionMechanic.cpp
@@ -530,11 +530,51 @@ void UMeshTopologySelectionMechanic::GrowSelection()
 		PersistentSelection.SelectedCornerIDs.Append(VerticesToAdd);
 	}
 
-	if (Properties->bSelectEdges || Properties->bSelectFaces)
+	if (Properties->bSelectEdges)
 	{
-		// TODO add support for growing edge/face selections if your tool requires it.
-		// growing edge/face selection not yet supported
-		checkNoEntry();
+		// grow edge selection
+		TSet<int32> EdgesToAdd;
+		for (const int32 EdgeIndex : PersistentSelection.SelectedEdgeIDs)
+		{
+			FDynamicMesh3::FEdge CurrentEdge = Mesh->GetEdge(EdgeIndex);
+			for (int32 EdgeVertex = 0; EdgeVertex<2; ++EdgeVertex)
+			{
+				for (int32 EdgeID : Mesh->VtxEdgesItr(CurrentEdge.Vert[EdgeVertex]))
+				{
+					if (!PersistentSelection.SelectedEdgeIDs.Contains(EdgeID))
+					{
+						EdgesToAdd.Add(EdgeID);
+					}
+				}
+			}
+		}
+		
+		PersistentSelection.SelectedEdgeIDs.Append(EdgesToAdd);
+	}
+
+	if (Properties->bSelectFaces)
+	{
+		// grow face selection
+		TSet<int32> TrianglesToAdd;
+		for (const int32 TriIndex : PersistentSelection.SelectedGroupIDs)
+		{
+			FIndex3i NeighborTriangles = Mesh->GetTriNeighbourTris(TriIndex);
+			for (int32 TriSideIndex = 0; TriSideIndex < 3; ++TriSideIndex)
+			{
+				const int32 NeighborTriangleIndex = NeighborTriangles[TriSideIndex];
+				if (NeighborTriangleIndex == FDynamicMesh3::InvalidID)
+				{
+					continue;
+				}
+				
+				if (!PersistentSelection.SelectedGroupIDs.Contains(NeighborTriangleIndex))
+				{
+					TrianglesToAdd.Add(NeighborTriangleIndex);
+				}
+			}
+		}
+		
+		PersistentSelection.SelectedGroupIDs.Append(TrianglesToAdd);
 	}
 
 	SelectionTimestamp++;
@@ -570,11 +610,57 @@ void UMeshTopologySelectionMechanic::ShrinkSelection()
 		}
 	}
 
-	if (Properties->bSelectEdges || Properties->bSelectFaces)
+	if (Properties->bSelectEdges)
 	{
-		// TODO add support for shrinking edge/face selections if your tool requires it.
-		// shrinking edge/face selection not yet supported
-		checkNoEntry();
+		// shrink edge selection
+		TSet<int32> BorderEdges;
+		for (const int32 EdgeIndex : PersistentSelection.SelectedEdgeIDs)
+		{
+			FDynamicMesh3::FEdge CurrentEdge = Mesh->GetEdge(EdgeIndex);
+			for (int32 EdgeVertex = 0; EdgeVertex<2; ++EdgeVertex)
+			{
+				for (const int32 EdgeID : Mesh->VtxEdgesItr(CurrentEdge.Vert[EdgeVertex]))
+				{
+					if (!PersistentSelection.SelectedEdgeIDs.Contains(EdgeID))
+					{
+						BorderEdges.Add(EdgeIndex);
+					}
+				}
+			}
+		}
+		
+		for (const int32 BorderEdge : BorderEdges)
+		{
+			PersistentSelection.SelectedEdgeIDs.Remove(BorderEdge);
+		}
+	}
+
+	if (Properties->bSelectFaces)
+	{
+		// shrink face selection
+		TSet<int32> BorderTriangles;
+		for (const int32 TriIndex : PersistentSelection.SelectedGroupIDs)
+		{
+			FIndex3i NeighborTriangles = Mesh->GetTriNeighbourTris(TriIndex);
+			for (int32 TriSideIndex = 0; TriSideIndex < 3; ++TriSideIndex)
+			{
+				const int32 NeighborTriangleIndex = NeighborTriangles[TriSideIndex];
+				if (NeighborTriangleIndex == FDynamicMesh3::InvalidID)
+				{
+					continue;
+				}
+				
+				if (!PersistentSelection.SelectedGroupIDs.Contains(NeighborTriangleIndex))
+				{
+					BorderTriangles.Add(TriIndex);
+				}
+			}
+		}
+		
+		for (const int32 BorderTriangle : BorderTriangles)
+		{
+			PersistentSelection.SelectedGroupIDs.Remove(BorderTriangle);
+		}
 	}
 
 	SelectionTimestamp++;
@@ -588,6 +674,7 @@ void UMeshTopologySelectionMechanic::FloodSelection()
 	ParentTool->GetToolManager()->BeginUndoTransaction(LOCTEXT("FloodSelectionChange", "Flood Selection"));
 	BeginChange();
 
+	// add all connected VERTICES to selection
 	if (Properties->bSelectVertices)
 	{
 		TSet<int32>& SelectedVertices = PersistentSelection.SelectedCornerIDs;
@@ -611,11 +698,63 @@ void UMeshTopologySelectionMechanic::FloodSelection()
 		}
 	}
 
-	if (Properties->bSelectEdges || Properties->bSelectFaces)
+	// add all connected EDGES to selection
+	if (Properties->bSelectEdges)
+	{
+		TSet<int32>& SelectedEdges = PersistentSelection.SelectedEdgeIDs;
+		TSet<int32> EdgesAddedInPrevIteration = PersistentSelection.SelectedEdgeIDs;
+		while(!EdgesAddedInPrevIteration.IsEmpty())
+		{
+			TSet<int32> EdgesToAddThisIteration;
+			for (const int32 EdgeAdded : EdgesAddedInPrevIteration)
+			{
+				FDynamicMesh3::FEdge CurrentEdge = Mesh->GetEdge(EdgeAdded);
+				for (int32 EdgeVertex = 0; EdgeVertex<2; ++EdgeVertex)
+				{
+					for (const int32 NeighborEdgeID : Mesh->VtxEdgesItr(CurrentEdge.Vert[EdgeVertex]))
+					{
+						if (!SelectedEdges.Contains(NeighborEdgeID))
+                    	{
+                    		EdgesToAddThisIteration.Add(NeighborEdgeID);
+                    	}
+					}
+				}
+			}
+			
+			SelectedEdges.Append(EdgesToAddThisIteration);
+			EdgesAddedInPrevIteration = EdgesToAddThisIteration;
+		}
+	}
+
+	// add all connected FACES to selection
+	if (Properties->bSelectFaces)
 	{
-		// TODO add support for flooding edge/face selections if your tool requires it.
-		// flooding edge/face selection not yet supported
-		checkNoEntry();
+		TSet<int32>& SelectedFaces = PersistentSelection.SelectedGroupIDs;
+		TSet<int32> FacesAddedInPrevIteration = PersistentSelection.SelectedGroupIDs;
+		while(!FacesAddedInPrevIteration.IsEmpty())
+		{
+			TSet<int32> FacesToAddThisIteration;
+			for (const int32 FaceAdded : FacesAddedInPrevIteration)
+			{
+				FIndex3i NeighborTriangles = Mesh->GetTriNeighbourTris(FaceAdded);
+				for (int32 TriSideIndex = 0; TriSideIndex < 3; ++TriSideIndex)
+				{
+					const int32 NeighborTriangleIndex = NeighborTriangles[TriSideIndex];
+					if (NeighborTriangleIndex == FDynamicMesh3::InvalidID)
+					{
+						continue;
+					}
+				
+					if (!PersistentSelection.SelectedGroupIDs.Contains(NeighborTriangleIndex))
+					{
+						FacesToAddThisIteration.Add(NeighborTriangleIndex);
+					}
+				}
+			}
+			
+			SelectedFaces.Append(FacesToAddThisIteration);
+			FacesAddedInPrevIteration = FacesToAddThisIteration;
+		}
 	}
 
 	SelectionTimestamp++;
