Improved component editing in skin paint tool

diff --git a/Engine/Plugins/Experimental/Animation/SkeletalMeshModelingTools/Source/SkeletalMeshModelingTools/Private/DetailCustomization/SkinWeightDetailCustomization.cpp b/Engine/Plugins/Experimental/Animation/SkeletalMeshModelingTools/Source/SkeletalMeshModelingTools/Private/DetailCustomization/SkinWeightDetailCustomization.cpp
--- a/Engine/Plugins/Experimental/Animation/SkeletalMeshModelingTools/Source/SkeletalMeshModelingTools/Private/DetailCustomization/SkinWeightDetailCustomization.cpp
+++ b/Engine/Plugins/Experimental/Animation/SkeletalMeshModelingTools/Source/SkeletalMeshModelingTools/Private/DetailCustomization/SkinWeightDetailCustomization.cpp
@@ -8,13 +8,10 @@
 #include "ModelingToolsEditorModeStyle.h"
 #include "Widgets/Input/SSegmentedControl.h"
 #include "SkeletalMesh/SkinWeightsPaintTool.h"
-#include "SSkinWeightProfileImportOptions.h"
 #include "Selection/PolygonSelectionMechanic.h"
-#include "UObject/UnrealTypePrivate.h"
 #include "Widgets/Input/SSpinBox.h"
 #include "Widgets/Input/SButton.h"
 #include "Widgets/Input/SNumericEntryBox.h"
-#include "Widgets/Input/SSlider.h"
 #include "Framework/MultiBox/MultiBoxBuilder.h"
 
 #define LOCTEXT_NAMESPACE "SkinWeightToolSettingsEditor"
@@ -25,6 +22,22 @@ float FSkinWeightDetailCustomization::WeightEditingLabelsPercent = 0.40f;
 float FSkinWeightDetailCustomization::WeightEditVerticalPadding = 4.0f;
 float FSkinWeightDetailCustomization::WeightEditHorizontalPadding = 2.0f;
 
+// add a buffer to the weight sliders to prevent ever fully getting a value to 1 or 0 using the slider alone
+// doing so will remove other influences and cause the slider to no longer function as all other influences
+// will be culled by normalization thus making the slider "stuck" at full value.
+constexpr float SliderBuffer = UE::AnimationCore::BoneWeightThreshold * 2.f;
+
+FSkinWeightDetailCustomization::~FSkinWeightDetailCustomization()
+{
+	if (Tool.IsValid())
+	{
+		Tool->OnSelectionChanged.RemoveAll(this);
+	}
+
+	Tool.Reset();
+	ToolSettings.Reset();
+}
+
 void FSkinWeightDetailCustomization::CustomizeDetails(IDetailLayoutBuilder& DetailBuilder)
 {
 	CurrentDetailBuilder = &DetailBuilder;
@@ -34,7 +47,9 @@ void FSkinWeightDetailCustomization::CustomizeDetails(IDetailLayoutBuilder& Deta
 
 	// should be impossible to get multiple settings objects for a single tool
 	ensure(DetailObjects.Num()==1);
-	SkinToolSettings = Cast<USkinWeightsPaintToolProperties>(DetailObjects[0]);
+	ToolSettings = Cast<USkinWeightsPaintToolProperties>(DetailObjects[0]);
+	Tool = ToolSettings->WeightTool;
+	Tool->OnSelectionChanged.AddSP(this, &FSkinWeightDetailCustomization::OnSelectionChanged);
 
 	// Edit SkinWeightLayer category first 
 	IDetailCategoryBuilder& SkinWeightLayerCategory = DetailBuilder.EditCategory("SkinWeightLayer", FText::GetEmpty(), ECategoryPriority::Important);
@@ -57,12 +72,12 @@ void FSkinWeightDetailCustomization::CustomizeDetails(IDetailLayoutBuilder& Deta
 					"Bones: select and manipulate bones to preview deformations.\n"))
 			.Value_Lambda([this]()
 			{
-				return SkinToolSettings->EditingMode;
+				return ToolSettings->EditingMode;
 			})
 			.OnValueChanged_Lambda([this](EWeightEditMode Mode)
 			{
-				SkinToolSettings->EditingMode = Mode;
-				SkinToolSettings->WeightTool->ToggleEditingMode();
+				ToolSettings->EditingMode = Mode;
+				ToolSettings->WeightTool->ToggleEditingMode();
 				if (CurrentDetailBuilder)
 				{
 					CurrentDetailBuilder->ForceRefreshDetails();
@@ -78,13 +93,13 @@ void FSkinWeightDetailCustomization::CustomizeDetails(IDetailLayoutBuilder& Deta
 	];
 
 	// BRUSH editing mode UI
-	if (SkinToolSettings->EditingMode == EWeightEditMode::Brush)
+	if (ToolSettings->EditingMode == EWeightEditMode::Brush)
 	{
 		AddBrushUI(DetailBuilder);
 	}
 
 	// MESH editing mode UI
-	if (SkinToolSettings->EditingMode == EWeightEditMode::Mesh)
+	if (ToolSettings->EditingMode == EWeightEditMode::Mesh)
 	{
 		AddSelectionUI(DetailBuilder);
 	}
@@ -108,11 +123,11 @@ void FSkinWeightDetailCustomization::CustomizeDetails(IDetailLayoutBuilder& Deta
 					"Full Material: Displays normal mesh materials with textures.\n"))
 			.Value_Lambda([this]()
 			{
-				return SkinToolSettings->ColorMode;
+				return ToolSettings->ColorMode;
 			})
 			.OnValueChanged_Lambda([this](EWeightColorMode Mode)
 			{
-				SkinToolSettings->SetColorMode(Mode);
+				ToolSettings->SetColorMode(Mode);
 			})
 			+ SSegmentedControl<EWeightColorMode>::Slot(EWeightColorMode::Greyscale)
 			.Text(LOCTEXT("GreyscaleMode", "Greyscale"))
@@ -146,9 +161,6 @@ void FSkinWeightDetailCustomization::CustomizeDetails(IDetailLayoutBuilder& Deta
 
 void FSkinWeightDetailCustomization::AddBrushUI(IDetailLayoutBuilder& DetailBuilder)
 {
-	// pointer to capture
-	TObjectPtr<USkinWeightsPaintToolProperties> ToolSettings = SkinToolSettings.Pin().Get();
-	
 	// custom display of falloff mode as segmented toggle buttons
 	IDetailCategoryBuilder& BrushCategory = DetailBuilder.EditCategory("Brush", FText::GetEmpty(), ECategoryPriority::Important);
 
@@ -168,11 +180,11 @@ void FSkinWeightDetailCustomization::AddBrushUI(IDetailLayoutBuilder& DetailBuil
 				"This command operates on the selected bone(s) and selected vertices.\n"
 				"If no bones are selected, ALL bones are considered.\n"
 				"If no vertices are selected, ALL vertices are considered."))
-			.Value_Lambda([ToolSettings]()
+			.Value_Lambda([this]()
 			{
 				return ToolSettings->BrushMode;
 			})
-			.OnValueChanged_Lambda([ToolSettings](EWeightEditOperation Mode)
+			.OnValueChanged_Lambda([this](EWeightEditOperation Mode)
 			{
 				ToolSettings->SetBrushMode(Mode);
 			})
@@ -198,11 +210,11 @@ void FSkinWeightDetailCustomization::AddBrushUI(IDetailLayoutBuilder& DetailBuil
 			.ToolTipText(LOCTEXT("BrushFalloffModeTooltip",
 					"Surface: falloff is based on the distance along the surface from the brush center to nearby connected vertices.\n"
 					"Volume: falloff is based on the straight-line distance from the brush center to surrounding vertices.\n"))
-			.Value_Lambda([ToolSettings]()
+			.Value_Lambda([this]()
 			{
 				return ToolSettings->GetBrushConfig().FalloffMode;
 			})
-			.OnValueChanged_Lambda([ToolSettings](EWeightBrushFalloffMode Mode)
+			.OnValueChanged_Lambda([this](EWeightBrushFalloffMode Mode)
 			{
 				ToolSettings->SetFalloffMode(Mode);
 			})
@@ -228,18 +240,18 @@ void FSkinWeightDetailCustomization::AddBrushUI(IDetailLayoutBuilder& DetailBuil
 		.MaxSliderValue(20.f)
 		.Value(10.0f)
 		.SupportDynamicSliderMaxValue(true)
-		.Value_Lambda([ToolSettings]()
+		.Value_Lambda([this]()
 		{
 			return ToolSettings->GetBrushConfig().Radius;
 		})
-		.OnValueChanged_Lambda([ToolSettings](float NewValue)
+		.OnValueChanged_Lambda([this](float NewValue)
 		{
 			ToolSettings->BrushRadius = NewValue;
 			ToolSettings->GetBrushConfig().Radius = NewValue;
 			FPropertyChangedEvent PropertyChangedEvent(UBrushBaseProperties::StaticClass()->FindPropertyByName(GET_MEMBER_NAME_CHECKED(UBrushBaseProperties, BrushRadius)));
 			ToolSettings->PostEditChangeProperty(PropertyChangedEvent);
 		})
-		.OnValueCommitted_Lambda([ToolSettings](float NewValue, ETextCommit::Type CommitType)
+		.OnValueCommitted_Lambda([this](float NewValue, ETextCommit::Type CommitType)
 		{
 			ToolSettings->SaveConfig();
 		})
@@ -261,18 +273,18 @@ void FSkinWeightDetailCustomization::AddBrushUI(IDetailLayoutBuilder& DetailBuil
 		.MaxSliderValue(1.f)
 		.Value(1.0f)
 		.SupportDynamicSliderMaxValue(true)
-		.Value_Lambda([ToolSettings]()
+		.Value_Lambda([this]()
 		{
 			return ToolSettings->GetBrushConfig().Strength;
 		})
-		.OnValueChanged_Lambda([ToolSettings](float NewValue)
+		.OnValueChanged_Lambda([this](float NewValue)
 		{
 			ToolSettings->BrushStrength = NewValue;
 			ToolSettings->GetBrushConfig().Strength = NewValue;
 			FPropertyChangedEvent PropertyChangedEvent(UBrushBaseProperties::StaticClass()->FindPropertyByName(GET_MEMBER_NAME_CHECKED(UBrushBaseProperties, BrushStrength)));
 			ToolSettings->PostEditChangeProperty(PropertyChangedEvent);
 		})
-		.OnValueCommitted_Lambda([ToolSettings](float NewValue, ETextCommit::Type CommitType)
+		.OnValueCommitted_Lambda([this](float NewValue, ETextCommit::Type CommitType)
 		{
 			ToolSettings->SaveConfig();
 		})
@@ -291,18 +303,18 @@ void FSkinWeightDetailCustomization::AddBrushUI(IDetailLayoutBuilder& DetailBuil
 		SNew(SSpinBox<float>)
 		.MinValue(0.f)
 		.MaxValue(1.f)
-		.Value_Lambda([ToolSettings]()
+		.Value_Lambda([this]()
 		{
 			return ToolSettings->GetBrushConfig().Falloff;
 		})
-		.OnValueChanged_Lambda([ToolSettings](float NewValue)
+		.OnValueChanged_Lambda([this](float NewValue)
 		{
 			ToolSettings->BrushFalloffAmount = NewValue;
 			ToolSettings->GetBrushConfig().Falloff = NewValue;
 			FPropertyChangedEvent PropertyChangedEvent(UBrushBaseProperties::StaticClass()->FindPropertyByName(GET_MEMBER_NAME_CHECKED(UBrushBaseProperties, BrushFalloffAmount)));
 			ToolSettings->PostEditChangeProperty(PropertyChangedEvent);
 		})
-		.OnValueCommitted_Lambda([ToolSettings](float NewValue, ETextCommit::Type CommitType)
+		.OnValueCommitted_Lambda([this](float NewValue, ETextCommit::Type CommitType)
 		{
 			ToolSettings->SaveConfig();
 		})
@@ -311,10 +323,6 @@ void FSkinWeightDetailCustomization::AddBrushUI(IDetailLayoutBuilder& DetailBuil
 
 void FSkinWeightDetailCustomization::AddSelectionUI(IDetailLayoutBuilder& DetailBuilder)
 {
-	// pointers to capture
-	TObjectPtr<USkinWeightsPaintToolProperties> ToolSettings = SkinToolSettings.Get();
-	TObjectPtr<USkinWeightsPaintTool> Tool = SkinToolSettings->WeightTool;
-	
 	// custom display of weight editing tools
 	IDetailCategoryBuilder& EditSelectionCategory = DetailBuilder.EditCategory("Edit Selection", FText::GetEmpty(), ECategoryPriority::Important);
 	EditSelectionCategory.InitiallyCollapsed(true);
@@ -327,18 +335,18 @@ void FSkinWeightDetailCustomization::AddSelectionUI(IDetailLayoutBuilder& Detail
 	ToolbarBuilder.BeginSection("SelectionFilter");
 	ToolbarBuilder.BeginBlockGroup();
 	
-	auto AddToggleButtonForBool = [&ToolbarBuilder, ToolSettings](EComponentSelectionMode Mode, const FText& Label, const FText& Tooltip, const FName IconName)
+	auto AddToggleButtonForBool = [&ToolbarBuilder, this](EComponentSelectionMode Mode, const FText& Label, const FText& Tooltip, const FName IconName)
 	{
 		ToolbarBuilder.AddToolBarButton(FUIAction(
-		FExecuteAction::CreateLambda([ToolSettings, Mode]()
+		FExecuteAction::CreateLambda([this, Mode]()
 		{
 			ToolSettings->SetComponentMode(Mode);
 		}),
-		FCanExecuteAction::CreateLambda([ToolSettings]()
+		FCanExecuteAction::CreateLambda([this]()
 		{
 			return ToolSettings->EditingMode == EWeightEditMode::Mesh;
 		}),
-		FIsActionChecked::CreateLambda([ToolSettings, Mode]()
+		FIsActionChecked::CreateLambda([this, Mode]()
 		{
 			return ToolSettings->ComponentSelectionMode == Mode;
 		})),
@@ -399,18 +407,18 @@ void FSkinWeightDetailCustomization::AddSelectionUI(IDetailLayoutBuilder& Detail
 					.ToolTipText(LOCTEXT("IsolateSelectedTooltip",
 							"Shows only the selected faces in the viewport.\n"
 							"Weight editing operations will not affect hidden vertices.\n "))
-					.IsEnabled_Lambda([Tool]()
+					.IsEnabled_Lambda([this]()
 					{
 						const bool bHasSelection = Tool->IsAnyComponentSelected();
 						const bool bAlreadyIsolatingSelection = Tool->IsSelectionIsolated();
 						return bHasSelection ||  bAlreadyIsolatingSelection;
 					})
-					.IsChecked_Lambda([Tool]()
+					.IsChecked_Lambda([this]()
 					{
 						const bool bAlreadyIsolatingSelection = Tool->IsSelectionIsolated();
 						return bAlreadyIsolatingSelection ? ECheckBoxState::Checked : ECheckBoxState::Unchecked;
 					})
-					.OnCheckStateChanged_Lambda([Tool](ECheckBoxState InCheckBoxState)
+					.OnCheckStateChanged_Lambda([this](ECheckBoxState InCheckBoxState)
 					{
 						if (InCheckBoxState == ECheckBoxState::Checked)
 						{
@@ -423,7 +431,7 @@ void FSkinWeightDetailCustomization::AddSelectionUI(IDetailLayoutBuilder& Detail
 					})
 					[
 						SNew(STextBlock)
-						.Text_Lambda([Tool]()
+						.Text_Lambda([this]()
 						{
 							if (Tool->IsSelectionIsolated())
 							{
@@ -450,7 +458,7 @@ void FSkinWeightDetailCustomization::AddSelectionUI(IDetailLayoutBuilder& Detail
 				.Text(LOCTEXT("GrowSelectionButtonLabel", "Grow"))
 				.ToolTipText(LOCTEXT("GrowSelectionTooltip",
 						"Grow the current selection by adding connected neighbors to current selection.\n"))
-				.OnClicked_Lambda([Tool]()
+				.OnClicked_Lambda([this]()
 				{
 					Tool->GrowSelection();
 					return FReply::Handled();
@@ -466,7 +474,7 @@ void FSkinWeightDetailCustomization::AddSelectionUI(IDetailLayoutBuilder& Detail
 				.Text(LOCTEXT("ShrinkSelectionButtonLabel", "Shrink"))
 				.ToolTipText(LOCTEXT("ShrinkSelectionTooltip",
 						"Shrink the current selection by removing components on the border of the current selection.\n"))
-				.OnClicked_Lambda([Tool]()
+				.OnClicked_Lambda([this]()
 				{
 					Tool->ShrinkSelection();
 					return FReply::Handled();
@@ -482,7 +490,7 @@ void FSkinWeightDetailCustomization::AddSelectionUI(IDetailLayoutBuilder& Detail
 				.Text(LOCTEXT("FloodSelectionButtonLabel", "Flood"))
 				.ToolTipText(LOCTEXT("FloodSelectionTooltip",
 						"Flood the current selection by adding all connected components to the current selection.\n"))
-				.OnClicked_Lambda([Tool]()
+				.OnClicked_Lambda([this]()
 				{
 					Tool->FloodSelection();
 					return FReply::Handled();
@@ -495,168 +503,258 @@ void FSkinWeightDetailCustomization::AddSelectionUI(IDetailLayoutBuilder& Detail
 	IDetailCategoryBuilder& EditWeightsCategory = DetailBuilder.EditCategory("Edit Weights", FText::GetEmpty(), ECategoryPriority::Important);
 	EditWeightsCategory.InitiallyCollapsed(true);
 
-	// FLOOD WEIGHTS category
-	EditWeightsCategory.AddCustomRow(LOCTEXT("FloodWeightsRow", "Flood"), false)
+	// FLOOD WEIGHTS SLIDER category
+	ToolSettings->DirectEditState.Reset();
+	EditWeightsCategory.AddCustomRow(LOCTEXT("FloodWeightsRow", "Flood Weights Slider"), false)
+	.WholeRowContent()
+	[
+		SNew(SHorizontalBox)
+
+		+SHorizontalBox::Slot()
+		[
+			SNew(SSegmentedControl<EWeightEditOperation>)
+			.ToolTipText(LOCTEXT("EditModeTooltip",
+				"Add: applies the current weight plus the flood value to the new weight.\n"
+				"Multiply: applies the current weight multiplied by the flood value to the new weight.\n"))
+			.Value_Lambda([this]()
+			{
+				return ToolSettings->DirectEditState.EditMode;
+			})
+			.OnValueChanged_Lambda([this](EWeightEditOperation Mode)
+			{
+				ToolSettings->DirectEditState.EditMode = Mode;
+				ToolSettings->DirectEditState.Reset();
+			})
+			+SSegmentedControl<EWeightEditOperation>::Slot(EWeightEditOperation::Add)
+			.Text(LOCTEXT("BrushAddMode", "Add"))
+			+ SSegmentedControl<EWeightEditOperation>::Slot(EWeightEditOperation::Multiply)
+			.Text(LOCTEXT("BrushMultiplyMode", "Multiply"))
+		]
+		
+		+SHorizontalBox::Slot()
+		.FillWidth(1)
+		[
+			SNew(SVerticalBox)
+
+			+SVerticalBox::Slot()
+			.AutoHeight()
+			[
+				SNew(SSpinBox<float>)
+				.Visibility_Lambda([this]()
+				{
+					const bool bIsVisible = ToolSettings->DirectEditState.EditMode != EWeightEditOperation::Relax;
+					return bIsVisible ? EVisibility::Visible : EVisibility::Collapsed;
+				})
+				.MinSliderValue_Lambda([this]()
+				{
+					return ToolSettings->DirectEditState.GetModeMinValue() + SliderBuffer;
+				})
+				.MaxSliderValue_Lambda([this]()
+				{
+					return ToolSettings->DirectEditState.GetModeMaxValue() - SliderBuffer;
+				})
+				.MinValue_Lambda([this]()
+				{
+					return ToolSettings->DirectEditState.GetModeMinValue();
+				})
+				.MaxValue_Lambda([this]()
+				{
+					return ToolSettings->DirectEditState.GetModeMaxValue();
+				})
+				.Value_Lambda([this]()
+				{
+					return ToolSettings->DirectEditState.CurrentValue;
+				})
+				.OnValueChanged_Lambda([this](float NewValue)
+				{
+					ToolSettings->DirectEditState.CurrentValue = NewValue;
+					
+					if (ToolSettings->DirectEditState.bInTransaction)
+					{
+						
+						float Value = NewValue;
+						if (ToolSettings->DirectEditState.EditMode == EWeightEditOperation::Add)
+						{
+							Value = NewValue - ToolSettings->DirectEditState.StartValue;
+						}
+						
+						constexpr bool bShouldTransact = false;
+						
+						Tool->EditWeightsOnVertices(
+							Tool->GetCurrentBoneIndex(),
+							Value,
+							ToolSettings->DirectEditState.EditMode,
+							Tool->GetVerticesToEdit(),
+							bShouldTransact);
+					}
+				})
+				.OnValueCommitted_Lambda([this](float NewValue, ETextCommit::Type CommitType)
+				{
+					if (!ToolSettings->DirectEditState.bInTransaction)
+					{
+						constexpr bool bShouldTransact = true;
+						Tool->EditWeightsOnVertices(
+							Tool->GetCurrentBoneIndex(),
+							NewValue, 
+							ToolSettings->DirectEditState.EditMode,
+							Tool->GetVerticesToEdit(),
+							bShouldTransact);
+					}
+					ToolSettings->DirectEditState.bInTransaction = false;
+				})
+				.OnBeginSliderMovement_Lambda([this]()
+				{
+					ToolSettings->DirectEditState.StartValue = ToolSettings->DirectEditState.CurrentValue;
+					ToolSettings->DirectEditState.bInTransaction = true;
+					Tool->BeginChange();
+				})
+				.OnEndSliderMovement_Lambda([this](float)
+				{
+					const FText TransactionLabel = LOCTEXT("FloodWeightChange", "Flood weights on vertices.");
+					Tool->EndChange(TransactionLabel);
+					ToolSettings->DirectEditState.bInTransaction = false;
+
+					// reset slider after each edit because multiplying operation is always relative to 1.0
+					if (ToolSettings->DirectEditState.EditMode == EWeightEditOperation::Multiply)
+					{
+						ToolSettings->DirectEditState.CurrentValue = 1.0f;
+						ToolSettings->DirectEditState.StartValue = 1.0f;
+					}
+				})
+				.ToolTipText(LOCTEXT("FloodWeightsToolTip", "Drag the slider to interactively adjust weights on the selected vertices."))
+			]	
+		]
+	];
+
+	// REPLACE WEIGHTS category
+	EditWeightsCategory.AddCustomRow(LOCTEXT("ReplaceWeightsRow", "Replace"), false)
 	.WholeRowContent()
 	[
 		SNew(SVerticalBox)
 
 		+ SVerticalBox::Slot()
-		.Padding(0.f, WeightEditVerticalPadding)
+		.Padding(WeightEditHorizontalPadding, WeightEditVerticalPadding)
 		[
 			SNew(SHorizontalBox)
-				
+					
 			+SHorizontalBox::Slot()
 			.VAlign(VAlign_Center)
 			.FillWidth(WeightEditingLabelsPercent)
 			[
 				SNew(STextBlock)
-				.Text(LOCTEXT("FloodAmountLabel", "Flood Amount"))
+				.Text(LOCTEXT("ReplaceLabel", "Replace"))
 				.Font(FAppStyle::GetFontStyle(TEXT("PropertyWindow.NormalFont")))
-				.ToolTipText(LOCTEXT("FloodAmountTooltip", "The amount of weight to apply in Flood operation."))
+				.ToolTipText(LOCTEXT("ReplaceTooltip", "Replace the current weight."))
 			]
 
 			+SHorizontalBox::Slot()
-			.MaxWidth(WeightSliderWidths)
 			[
 				SNew(SSpinBox<float>)
 				.MinValue(0.f)
-				.MaxValue(2.0f)
-				.MaxSliderValue(1.f)
-				.Value(1.0f)
-				.SupportDynamicSliderMaxValue(true)
-				.Value_Lambda([ToolSettings]()
+				.MaxValue(1.f)
+				.Value_Lambda([this]()
 				{
-					return ToolSettings->FloodValue;
+					return ToolSettings->ReplaceValue;
 				})
-				.OnValueChanged_Lambda([ToolSettings](float NewValue)
+				.OnValueChanged_Lambda([this](float NewValue)
 				{
-					ToolSettings->FloodValue = NewValue;
+					ToolSettings->ReplaceValue = NewValue;
 				})
-				.OnValueCommitted_Lambda([ToolSettings](float NewValue, ETextCommit::Type CommitType)
+				.OnValueCommitted_Lambda([this](float NewValue, ETextCommit::Type CommitType)
 				{
 					ToolSettings->SaveConfig();
 				})
 			]
+
+			+SHorizontalBox::Slot()
+			[
+				SNew(SBox)
+				[
+					SNew(SButton)
+					.HAlign(HAlign_Center)
+					.Text(LOCTEXT("ReplaceWeightsButtonLabel", "Replace"))
+					.ToolTipText(LOCTEXT("ReplaceButtonTooltip",
+						"Replace: the weight of selected vertices is replaced by the specified value.\n"))
+					.OnClicked_Lambda([this]()
+					{
+						constexpr bool bShouldTransact = true;
+						Tool->EditWeightsOnVertices(
+							Tool->GetCurrentBoneIndex(),
+							ToolSettings->ReplaceValue, 
+							EWeightEditOperation::Replace,
+							Tool->GetVerticesToEdit(),
+							bShouldTransact);
+						return FReply::Handled();
+					})
+				]
+			]
 		]
+	];
 
-		+SVerticalBox::Slot()
-		.Padding(0.f, WeightEditVerticalPadding)
+	// RELAX WEIGHTS category
+	EditWeightsCategory.AddCustomRow(LOCTEXT("RelaxWeightsRow", "Relax"), false)
+	.WholeRowContent()
+	[
+		SNew(SVerticalBox)
+
+		+ SVerticalBox::Slot()
+		.Padding(WeightEditHorizontalPadding, WeightEditVerticalPadding)
 		[
 			SNew(SHorizontalBox)
-			
+					
 			+SHorizontalBox::Slot()
 			.VAlign(VAlign_Center)
 			.FillWidth(WeightEditingLabelsPercent)
 			[
 				SNew(STextBlock)
-				.Text(LOCTEXT("FloodOperationLabel", "Flood Operation"))
+				.Text(LOCTEXT("RelaxIterationsLabel", "Iterations"))
 				.Font(FAppStyle::GetFontStyle(TEXT("PropertyWindow.NormalFont")))
-				.ToolTipText(LOCTEXT("FloodOperationTooltip", "Various weight editing operations applied to the current selection."))
+				.ToolTipText(LOCTEXT("PruneThresholdTooltip", "The number of iterations of relaxation to apply."))
 			]
 
 			+SHorizontalBox::Slot()
 			[
-				SNew(SHorizontalBox)
-				+ SHorizontalBox::Slot()
-				[
-					SNew(SButton)
-					.HAlign(HAlign_Center)
-					.Text(LOCTEXT("AddWeightsButtonLabel", "Add"))
-					.ToolTipText(LOCTEXT("AddOpTooltip", "Add: adds the Flood Amount value to the current weight."))
-					.OnClicked_Lambda([ToolSettings, Tool]()
-					{
-						Tool->FloodWeights(ToolSettings->FloodValue, EWeightEditOperation::Add);
-						return FReply::Handled();
-					})
-				]
-				+ SHorizontalBox::Slot()
-				[
-					SNew(SButton)
-					.HAlign(HAlign_Center)
-					.Text(LOCTEXT("ReplaceWeightsButtonLabel", "Replace"))
-					.ToolTipText(LOCTEXT("ReplaceOpTooltip", "Replace: applies the current weight minus the flood amount to the new weight."))
-					.OnClicked_Lambda([ToolSettings, Tool]()
-					{
-						Tool->FloodWeights(ToolSettings->FloodValue, EWeightEditOperation::Replace);
-						return FReply::Handled();
-					})
-				]
-				+ SHorizontalBox::Slot()
-				[
-					SNew(SButton)
-					.HAlign(HAlign_Center)
-					.Text(LOCTEXT("MultiplyeightsButtonLabel", "Multiply"))
-					.ToolTipText(LOCTEXT("MultiplyOpTooltip", "Multiply: applies the current weight multiplied by the flood amount to the new weight."))
-					.OnClicked_Lambda([ToolSettings, Tool]()
-					{
-						Tool->FloodWeights(ToolSettings->FloodValue, EWeightEditOperation::Multiply);
-						return FReply::Handled();
-					})
-				]
-				+ SHorizontalBox::Slot()
+				SNew(SSpinBox<int32>)
+				.MinValue(1)
+				.MaxValue(10)
+				.Value_Lambda([this]()
+				{
+					return ToolSettings->RelaxIterations;
+				})
+				.OnValueChanged_Lambda([this](int32 NewValue)
+				{
+					ToolSettings->RelaxIterations = NewValue;
+				})
+				.OnValueCommitted_Lambda([this](int32 NewValue, ETextCommit::Type CommitType)
+				{
+					ToolSettings->SaveConfig();
+				})
+			]
+
+			+SHorizontalBox::Slot()
+			[
+				SNew(SBox)
 				[
 					SNew(SButton)
 					.HAlign(HAlign_Center)
 					.Text(LOCTEXT("RelaxWeightsButtonLabel", "Relax"))
-					.ToolTipText(LOCTEXT("RelaxeOpTooltip", "Relax: applies the average of the connected (by edge) vertex weights to the new vertex weight, scaled by the flood amount."))
-					.OnClicked_Lambda([ToolSettings, Tool]()
+					.ToolTipText(LOCTEXT("PruneButtonTooltip",
+						"Relax: the weight of each vertex is replaced by the average of it's neighbors. This smooths weights across the mesh.\n"))
+					.OnClicked_Lambda([this]()
 					{
-						Tool->FloodWeights(ToolSettings->FloodValue, EWeightEditOperation::Relax);
+						constexpr bool bShouldTransact = true;
+						Tool->EditWeightsOnVertices(
+							Tool->GetCurrentBoneIndex(),
+							ToolSettings->RelaxIterations, 
+							EWeightEditOperation::Relax,
+							Tool->GetVerticesToEdit(),
+							bShouldTransact);
 						return FReply::Handled();
 					})
 				]
 			]
 		]
 	];
-
-	// AVERAGE/NORMALIZE WEIGHTS category
-	EditWeightsCategory.AddCustomRow(LOCTEXT("NormalizeWeightsRow", "Normalize"), false)
-	.WholeRowContent()
-	[
-		SNew(SHorizontalBox)
-
-		+SHorizontalBox::Slot()
-		.Padding(2.f, WeightEditVerticalPadding)
-		[
-			SNew(SButton)
-			.HAlign(HAlign_Center)
-			.VAlign(VAlign_Center)
-			.Text(LOCTEXT("AverageWeightsButtonLabel", "Average"))
-			.ToolTipText(LOCTEXT("AverageWeightsTooltip",
-					"Takes the average of vertex weights and applies it to the selected vertices.\n"
-					"This command operates on the selected bone(s) and selected vertices.\n "
-					"If no bones are selected, ALL bone weights are considered.\n "
-					"If no vertices are selected, ALL vertices are considered."))
-			.OnClicked_Lambda([Tool]()
-			{
-				Tool->AverageWeights();
-				return FReply::Handled();
-			})
-		]
-
-		+SHorizontalBox::Slot()
-		.Padding(2.f, WeightEditVerticalPadding)
-		[
-			SNew(SButton)
-			.HAlign(HAlign_Center)
-			.VAlign(VAlign_Center)
-			.Text(LOCTEXT("NormalizeWeightsButtonLabel", "Normalize"))
-			.ToolTipText(LOCTEXT("NormalizeWeightsTooltip",
-					"Forces the weights on the selected vertices to sum to 1.\n"
-					"This command operates on the selected vertices.\n "
-					"If no vertices are selected, ALL vertices are considered."))
-			.IsEnabled_Lambda([ToolSettings]()
-			{
-				return ToolSettings->EditingMode == EWeightEditMode::Mesh;
-			})
-			.OnClicked_Lambda([Tool]()
-			{
-				Tool->NormalizeWeights();
-				return FReply::Handled();
-			})
-		]
-	];
 	
 	// PRUNE WEIGHTS category
 	EditWeightsCategory.AddCustomRow(LOCTEXT("PruneWeightsRow", "Prune"), false)
@@ -684,15 +782,15 @@ void FSkinWeightDetailCustomization::AddSelectionUI(IDetailLayoutBuilder& Detail
 				SNew(SSpinBox<float>)
 				.MinValue(0.f)
 				.MaxValue(1.f)
-				.Value_Lambda([ToolSettings]()
+				.Value_Lambda([this]()
 				{
 					return ToolSettings->PruneValue;
 				})
-				.OnValueChanged_Lambda([ToolSettings](float NewValue)
+				.OnValueChanged_Lambda([this](float NewValue)
 				{
 					ToolSettings->PruneValue = NewValue;
 				})
-				.OnValueCommitted_Lambda([ToolSettings](float NewValue, ETextCommit::Type CommitType)
+				.OnValueCommitted_Lambda([this](float NewValue, ETextCommit::Type CommitType)
 				{
 					ToolSettings->SaveConfig();
 				})
@@ -712,7 +810,7 @@ void FSkinWeightDetailCustomization::AddSelectionUI(IDetailLayoutBuilder& Detail
 						"If no vertices are selected, ALL vertices are considered."))
 					.OnClicked_Lambda([this]()
 					{
-						SkinToolSettings->WeightTool->PruneWeights(SkinToolSettings->PruneValue);
+						ToolSettings->WeightTool->PruneWeights(ToolSettings->PruneValue);
 						return FReply::Handled();
 					})
 				]
@@ -720,6 +818,54 @@ void FSkinWeightDetailCustomization::AddSelectionUI(IDetailLayoutBuilder& Detail
 		]
 	];
 
+	// AVERAGE/NORMALIZE WEIGHTS category
+	EditWeightsCategory.AddCustomRow(LOCTEXT("NormalizeWeightsRow", "Normalize"), false)
+	.WholeRowContent()
+	[
+		SNew(SHorizontalBox)
+
+		+SHorizontalBox::Slot()
+		.Padding(2.f, WeightEditVerticalPadding)
+		[
+			SNew(SButton)
+			.HAlign(HAlign_Center)
+			.VAlign(VAlign_Center)
+			.Text(LOCTEXT("AverageWeightsButtonLabel", "Average"))
+			.ToolTipText(LOCTEXT("AverageWeightsTooltip",
+					"Takes the average of vertex weights and applies it to the selected vertices.\n"
+					"This command operates on the selected bone(s) and selected vertices.\n "
+					"If no bones are selected, ALL bone weights are considered.\n "
+					"If no vertices are selected, ALL vertices are considered."))
+			.OnClicked_Lambda([this]()
+			{
+				Tool->AverageWeights();
+				return FReply::Handled();
+			})
+		]
+
+		+SHorizontalBox::Slot()
+		.Padding(2.f, WeightEditVerticalPadding)
+		[
+			SNew(SButton)
+			.HAlign(HAlign_Center)
+			.VAlign(VAlign_Center)
+			.Text(LOCTEXT("NormalizeWeightsButtonLabel", "Normalize"))
+			.ToolTipText(LOCTEXT("NormalizeWeightsTooltip",
+					"Forces the weights on the selected vertices to sum to 1.\n"
+					"This command operates on the selected vertices.\n "
+					"If no vertices are selected, ALL vertices are considered."))
+			.IsEnabled_Lambda([this]()
+			{
+				return ToolSettings->EditingMode == EWeightEditMode::Mesh;
+			})
+			.OnClicked_Lambda([this]()
+			{
+				Tool->NormalizeWeights();
+				return FReply::Handled();
+			})
+		]
+	];
+
 	// MIRROR WEIGHTS category
 	EditWeightsCategory.AddCustomRow(LOCTEXT("MirrorWeightsRow", "Mirror"), false)
 	.WholeRowContent()
@@ -752,11 +898,11 @@ void FSkinWeightDetailCustomization::AddSelectionUI(IDetailLayoutBuilder& Detail
 						"X: copies weights across the YZ plane.\n"
 						"Y: copies weights across the XZ plane.\n"
 						"Z: copies weights across the XY plane."))
-					.Value_Lambda([ToolSettings]()
+					.Value_Lambda([this]()
 					{
 						return ToolSettings->MirrorAxis;
 					})
-					.OnValueChanged_Lambda([ToolSettings](EAxis::Type Mode)
+					.OnValueChanged_Lambda([this](EAxis::Type Mode)
 					{
 						ToolSettings->MirrorAxis = Mode;
 					})
@@ -772,11 +918,11 @@ void FSkinWeightDetailCustomization::AddSelectionUI(IDetailLayoutBuilder& Detail
 				[
 					SNew(SSegmentedControl<EMirrorDirection>)
 					.ToolTipText(LOCTEXT("MirrorDirectionTooltip", "The direction that determines what side of the plane to copy weights from."))
-					.Value_Lambda([ToolSettings]()
+					.Value_Lambda([this]()
 					{
 						return ToolSettings->MirrorDirection;
 					})
-					.OnValueChanged_Lambda([ToolSettings](EMirrorDirection Mode)
+					.OnValueChanged_Lambda([this](EMirrorDirection Mode)
 					{
 						ToolSettings->MirrorDirection = Mode;
 					})
@@ -801,7 +947,7 @@ void FSkinWeightDetailCustomization::AddSelectionUI(IDetailLayoutBuilder& Detail
 					"This command operates on the selected bone(s) and selected vertices.\n "
 					"If no bones are selected, ALL bone weights are considered.\n "
 					"If no vertices are selected, ALL vertices are considered."))
-				.OnClicked_Lambda([ToolSettings, Tool]()
+				.OnClicked_Lambda([this]()
 				{
 					Tool->MirrorWeights(ToolSettings->MirrorAxis, ToolSettings->MirrorDirection);
 					return FReply::Handled();
@@ -814,7 +960,7 @@ void FSkinWeightDetailCustomization::AddSelectionUI(IDetailLayoutBuilder& Detail
 	EditWeightsCategory.AddCustomRow(LOCTEXT("VertexEditorRow", "Component Editor"), false)
 	.WholeRowContent()
 	[
-		SNew(SVertexWeightEditor, SkinToolSettings->WeightTool)
+		SNew(SVertexWeightEditor, ToolSettings->WeightTool)
 	];
 }
 
@@ -835,11 +981,6 @@ TSharedRef<SWidget> SVertexWeightItem::GenerateWidgetForColumn(const FName& Colu
 
 	if (ColumnName == "Weight")
 	{
-		// add a buffer to the slider to prevent ever fully getting a value to 1 or 0 using the slider alone
-		// doing so will remove other influences and cause the slider to no longer function as all other influences
-		// will be culled by normalization thus making the slider "stuck" at full value.
-		constexpr float SliderBuffer = 0.001f;
-		
 		return SNew(SNumericEntryBox<float>)
 		.AllowSpin(true)
 		.MinSliderValue(SliderBuffer)
@@ -848,26 +989,43 @@ TSharedRef<SWidget> SVertexWeightItem::GenerateWidgetForColumn(const FName& Colu
 		.MaxValue(1.f)
 		.Value_Lambda([this]()
 		{
-			return ParentTable->Tool->GetAverageWeightOnBone(Element->BoneIndex, ParentTable->SelectedVertices);
+			if (bInTransaction)
+			{
+				return ValueDuringSlide;
+			}
+			
+			return ParentTable->Tool->GetAverageWeightOnBone(Element->BoneIndex, ParentTable->Tool->GetVerticesToEdit());
 		})
 		.OnValueChanged_Lambda([this](float NewValue)
 		{
-			TArray<VertexIndex> VerticesToEdit;
-			ParentTable->Tool->GetSelectedVertices(VerticesToEdit);
-			ParentTable->Tool->SetBoneWeightOnVertices(Element->BoneIndex, NewValue, VerticesToEdit, !bInTransaction);
+			if (bInTransaction)
+			{
+				ValueDuringSlide = NewValue;
+				
+				const float ScaleMultiplier = NewValue / FMath::Max(ValueAtStartOfSlide, SliderBuffer);
+				constexpr bool bShouldTransact = false;
+				ParentTable->Tool->EditWeightsOnVertices(Element->BoneIndex, ScaleMultiplier, EWeightEditOperation::Multiply, ParentTable->Tool->GetVerticesToEdit(), bShouldTransact);
+			}
 		})
 		.OnValueCommitted_Lambda([this](float NewValue, ETextCommit::Type CommitType)
 		{
+			if (!bInTransaction)
+			{
+				constexpr bool bShouldTransact = true;
+				ParentTable->Tool->EditWeightsOnVertices(Element->BoneIndex, NewValue, EWeightEditOperation::Replace, ParentTable->Tool->GetVerticesToEdit(), bShouldTransact);
+			}
 			bInTransaction = false;
 		})
 		.OnBeginSliderMovement_Lambda([this]()
 		{
 			ParentTable->Tool->BeginChange();
+			ValueAtStartOfSlide = ParentTable->Tool->GetAverageWeightOnBone(Element->BoneIndex, ParentTable->Tool->GetVerticesToEdit());
+			ValueDuringSlide = ValueAtStartOfSlide;
 			bInTransaction = true;
 		})
 		.OnEndSliderMovement_Lambda([this](float)
 		{
-			const FText TransactionLabel = LOCTEXT("DirectWeightChange", "Set weights on vertices.");
+			const FText TransactionLabel = LOCTEXT("DirectWeightChange", "Scale weights on vertices.");
 			ParentTable->Tool->EndChange(TransactionLabel);
 			bInTransaction = false;
 		})
@@ -924,14 +1082,10 @@ void SVertexWeightEditor::RefreshView()
 	{
 		return; 
 	}
-
-	// get list of selected vertex indices
-	SelectedVertices.Reset();
-	Tool->GetSelectedVertices(SelectedVertices);
 	
 	// get all bones affecting the selected vertices
 	TArray<int32> Influences;
-	Tool->GetInfluences(SelectedVertices, Influences);
+	Tool->GetInfluences(Tool->GetVerticesToEdit(), Influences);
 
 	// generate list view items
 	ListViewItems.Reset();
diff --git a/Engine/Plugins/Experimental/Animation/SkeletalMeshModelingTools/Source/SkeletalMeshModelingTools/Private/DetailCustomization/SkinWeightDetailCustomization.h b/Engine/Plugins/Experimental/Animation/SkeletalMeshModelingTools/Source/SkeletalMeshModelingTools/Private/DetailCustomization/SkinWeightDetailCustomization.h
--- a/Engine/Plugins/Experimental/Animation/SkeletalMeshModelingTools/Source/SkeletalMeshModelingTools/Private/DetailCustomization/SkinWeightDetailCustomization.h
+++ b/Engine/Plugins/Experimental/Animation/SkeletalMeshModelingTools/Source/SkeletalMeshModelingTools/Private/DetailCustomization/SkinWeightDetailCustomization.h
@@ -4,6 +4,7 @@
 
 #include "CoreMinimal.h"
 #include "IDetailCustomization.h"
+#include "SkeletalMesh/SkinWeightsPaintTool.h"
 #include "UObject/WeakObjectPtr.h"
 #include "Widgets/Views/STableRow.h"
 #include "Widgets/Views/SListView.h"
@@ -15,6 +16,8 @@ class FSkinWeightDetailCustomization : public IDetailCustomization
 {
 public:
 
+	virtual ~FSkinWeightDetailCustomization() override;
+
 	static TSharedRef<IDetailCustomization> MakeInstance()
 	{
 		return MakeShareable(new FSkinWeightDetailCustomization);
@@ -23,13 +26,19 @@ public:
 	/** IDetailCustomization interface */
 	virtual void CustomizeDetails(IDetailLayoutBuilder& DetailBuilder) override;
 
+	void OnSelectionChanged()
+	{
+		ToolSettings.Get()->DirectEditState.Reset();
+	};
+
 private:
 
 	void AddBrushUI(IDetailLayoutBuilder& DetailBuilder);
 	void AddSelectionUI(IDetailLayoutBuilder& DetailBuilder);
 
 	IDetailLayoutBuilder* CurrentDetailBuilder;
-	TWeakObjectPtr<USkinWeightsPaintToolProperties> SkinToolSettings;
+	TWeakObjectPtr<USkinWeightsPaintToolProperties> ToolSettings;
+	TWeakObjectPtr<USkinWeightsPaintTool> Tool;
 
 	static float WeightSliderWidths;
 	static float WeightEditingLabelsPercent;
@@ -64,6 +73,8 @@ private:
 	TSharedPtr<FWeightEditorElement> Element;
 	TSharedPtr<SVertexWeightEditor> ParentTable;
 	bool bInTransaction = false;
+	float ValueAtStartOfSlide;
+	float ValueDuringSlide;
 };
 
 typedef SListView< TSharedPtr<FWeightEditorElement> > SWeightEditorListViewType;
@@ -86,7 +97,6 @@ private:
 	TSharedPtr<SWeightEditorListViewType> ListView;
 	TArray< TSharedPtr<FWeightEditorElement> > ListViewItems;
 	TWeakObjectPtr<USkinWeightsPaintTool> Tool;
-	TArray<int32> SelectedVertices;
 
 	friend SVertexWeightItem;
 };
\ No newline at end of file
diff --git a/Engine/Plugins/Experimental/MeshModelingToolsetExp/Source/MeshModelingToolsEditorOnlyExp/Private/SkeletalMesh/SkinWeightsPaintTool.cpp b/Engine/Plugins/Experimental/MeshModelingToolsetExp/Source/MeshModelingToolsEditorOnlyExp/Private/SkeletalMesh/SkinWeightsPaintTool.cpp
--- a/Engine/Plugins/Experimental/MeshModelingToolsetExp/Source/MeshModelingToolsEditorOnlyExp/Private/SkeletalMesh/SkinWeightsPaintTool.cpp
+++ b/Engine/Plugins/Experimental/MeshModelingToolsetExp/Source/MeshModelingToolsEditorOnlyExp/Private/SkeletalMesh/SkinWeightsPaintTool.cpp
@@ -157,6 +157,48 @@ protected:
 	const SkinPaintTool::FSkinToolWeights* Weights = nullptr;
 };
 
+void FDirectEditWeightState::Reset()
+{
+	bInTransaction = false;
+	StartValue = CurrentValue = GetModeDefaultValue();
+}
+
+float FDirectEditWeightState::GetModeDefaultValue()
+{
+	static TMap<EWeightEditOperation, float> DefaultModeValues = {
+		{EWeightEditOperation::Add, 0.0f},
+		{EWeightEditOperation::Replace, .0f},
+		{EWeightEditOperation::Multiply, 1.f},
+		{EWeightEditOperation::Relax, 0.f}
+	};
+
+	return DefaultModeValues[EditMode];
+}
+
+float FDirectEditWeightState::GetModeMinValue()
+{
+	static TMap<EWeightEditOperation, float> MinModeValues = {
+		{EWeightEditOperation::Add, -1.f},
+		{EWeightEditOperation::Replace, 0.f},
+		{EWeightEditOperation::Multiply, 0.f},
+		{EWeightEditOperation::Relax, 0.f}
+	};
+
+	return MinModeValues[EditMode];
+}
+
+float FDirectEditWeightState::GetModeMaxValue()
+{
+	static TMap<EWeightEditOperation, float> MaxModeValues = {
+		{EWeightEditOperation::Add, 1.f},
+		{EWeightEditOperation::Replace, 1.f},
+		{EWeightEditOperation::Multiply, 2.f},
+		{EWeightEditOperation::Relax, 10.f}
+	};
+
+	return MaxModeValues[EditMode];
+}
+
 USkinWeightsPaintToolProperties::USkinWeightsPaintToolProperties()
 {
 	ColorRamp.Add(FLinearColor::Blue);
@@ -267,6 +309,16 @@ float FMultiBoneWeightEdits::GetVertexDeltaFromEdits(const int32 BoneIndex, cons
 	return 0.0f;
 }
 
+void FMultiBoneWeightEdits::GetEditedVertexIndices(TSet<int32>& OutVerticesToEdit) const
+{
+	TSet<int32> VerticesInEdit;
+	for (const TTuple<int32, FSingleBoneWeightEdits>& Pair : PerBoneWeightEdits)
+	{
+		Pair.Value.NewWeights.GetKeys(VerticesInEdit);
+		OutVerticesToEdit.Append(VerticesInEdit);
+	}
+}
+
 void FSkinToolDeformer::Initialize(const USkeletalMeshComponent* InSkelMeshComponent, const FMeshDescription* InMeshDescription)
 {
 	// get all bone transforms in the reference pose store a copy in component space
@@ -530,7 +582,7 @@ void FSkinToolWeights::EditVertexWeightAndNormalize(
 	TArray<int32> BonesAffectingVertex;
 	TArray<float> ValuesToNormalize;
 	float Total = 0.0f;
-	const VertexWeights& VertexData = CurrentWeights[VertexID];
+	const VertexWeights& VertexData = PreChangeWeights[VertexID];
 	for (const FVertexBoneWeight& VertexBoneData : VertexData)
 	{
 		if (VertexBoneData.BoneIndex == BoneToHoldIndex)
@@ -579,14 +631,20 @@ void FSkinToolWeights::EditVertexWeightAndNormalize(
 	// is the user trying to prune ALL weight from this vertex, AND all other weights are equal to zero?
 	// in this case, we have two options:
 	// 1. if there are other influences recorded for this vertex, then split remaining influence among them
-	// 2. if there are NO other influences recorded for this vertex, then move the weight to the root as last ditch effort
+	// 2. if there are NO other influences recorded for this vertex, then move the weight to the parent as last ditch effort
 	const bool bApplyZeroWeightToThisVertex = NewWeightValue <= MinimumWeightThreshold;
 	if (bApplyZeroWeightToThisVertex && bVertexHasNoOtherInfluences)
 	{
 		if (ValuesToNormalize.IsEmpty())
 		{
-			// assign all weight to the root
-			constexpr int32 BoneIndex = 0;
+			// assign all weight to the parent
+			int32 BoneIndex = 0;
+			if (const USkeletalMesh* SkeletalMesh = Deformer.Component->GetSkeletalMeshAsset())
+			{
+				BoneIndex = SkeletalMesh->GetRefSkeleton().GetParentIndex(BoneToHoldIndex);
+				BoneIndex = BoneIndex == INDEX_NONE ? 0 : BoneIndex;
+			}
+			
 			constexpr float OldWeight = 0.f;
 			constexpr float NewWeight = 1.f;
 			WeightEdits.MergeSingleEdit(BoneIndex, VertexID, OldWeight, NewWeight);
@@ -960,6 +1018,7 @@ void USkinWeightsPaintTool::Setup()
 	// notify preview mesh when triangle selection has been updated
 	PolygonSelectionMechanic->OnSelectionChanged.AddWeakLambda(this, [this]()
 	{
+		UpdateSelectedVertices();
 		PreviewMesh->FastNotifySecondaryTrianglesChanged();
 	});
 	PolygonSelectionMechanic->OnFaceSelectionPreviewChanged.AddWeakLambda(this, [this]()
@@ -1037,13 +1096,12 @@ FBox USkinWeightsPaintTool::GetWorldSpaceFocusBox()
 	case EWeightEditMode::Mesh:
 		{
 			FAxisAlignedBox3d Bounds = FAxisAlignedBox3d::Empty();
-			TArray<int32> SelectedVertexIndices;
-			GetSelectedVertices(SelectedVertexIndices);
-			if (!SelectedVertexIndices.IsEmpty())
+			UpdateSelectedVertices();
+			if (!SelectedVertices.IsEmpty())
 			{
 				const FDynamicMesh3* Mesh = PreviewMesh->GetMesh();
 				const FTransform3d Transform(PreviewMesh->GetTransform());
-				for (const int32 VertexID : SelectedVertexIndices)
+				for (const int32 VertexID : SelectedVertices)
 				{
 					Bounds.Contain(Transform.TransformPosition(Mesh->GetVertex(VertexID)));
 				}
@@ -1169,11 +1227,16 @@ void USkinWeightsPaintTool::OnTick(float DeltaTime)
 		PendingCurrentBone.Reset();
 	}
 
-	if (bVisibleWeightsValid == false || WeightToolProperties->bColorModeChanged)
+	if (bVertexColorsNeedUpdated)
+	{
+		UpdateVertexColorForAllVertices();
+		bVertexColorsNeedUpdated = false;
+	}
+
+	if (!VerticesToUpdateColor.IsEmpty())
 	{
-		UpdateCurrentBoneVertexColors();
-		bVisibleWeightsValid = true;
-		WeightToolProperties->bColorModeChanged = false;
+		UpdateVertexColorForSubsetOfVertices();
+		VerticesToUpdateColor.Empty();
 	}
 
 	// sparsely updates vertex positions (only on vertices with modified weights)
@@ -1210,7 +1273,7 @@ void USkinWeightsPaintTool::PostEditMeshInitialization(
 	}
 	Weights.Profile = WeightToolProperties->ActiveSkinWeightProfile;
 	Weights.InitializeSkinWeights(InComponent, &InMeshDescription);
-	bVisibleWeightsValid = false;
+	bVertexColorsNeedUpdated = true;
 	
 	// update smooth operator (this must be done after PreviewMesh & Weights have been updated)
 	InitializeSmoothWeightsOperator();
@@ -1242,7 +1305,6 @@ void USkinWeightsPaintToolProperties::SetComponentMode(EComponentSelectionMode I
 
 void USkinWeightsPaintToolProperties::SetFalloffMode(EWeightBrushFalloffMode InFalloffMode)
 {
-	bColorModeChanged = true;
 	GetBrushConfig().FalloffMode = InFalloffMode;
 	SaveConfig();
 
@@ -1252,8 +1314,6 @@ void USkinWeightsPaintToolProperties::SetFalloffMode(EWeightBrushFalloffMode InF
 void USkinWeightsPaintToolProperties::SetColorMode(EWeightColorMode InColorMode)
 {
 	ColorMode = InColorMode;
-	bColorModeChanged = true;
-
 	WeightTool->SetDisplayVertexColors(ColorMode!=EWeightColorMode::FullMaterial);
 	WeightTool->SetFocusInViewport();
 }
@@ -1556,8 +1616,10 @@ FVector4f USkinWeightsPaintTool::GetColorOfVertex(VertexIndex InVertexIndex, Bon
 }
 
 
-void USkinWeightsPaintTool::UpdateCurrentBoneVertexColors()
+void USkinWeightsPaintTool::UpdateVertexColorForAllVertices()
 {
+	TRACE_CPUPROFILER_EVENT_SCOPE(SkinTool::UpdateVertexColors);
+	
 	const int32 CurrentBoneIndex = GetBoneIndexFromName(CurrentBone);
 	
 	// update mesh with new value colors
@@ -1576,6 +1638,31 @@ void USkinWeightsPaintTool::UpdateCurrentBoneVertexColors()
 	PreviewMesh->NotifyDeferredEditCompleted(UPreviewMesh::ERenderUpdateMode::FastUpdate, EMeshRenderAttributeFlags::VertexColors, false);
 }
 
+void USkinWeightsPaintTool::UpdateVertexColorForSubsetOfVertices()
+{
+	TRACE_CPUPROFILER_EVENT_SCOPE(SkinTool::UpdateVertexColors);
+	
+	PreviewMesh->DeferredEditMesh([this](FDynamicMesh3& Mesh)
+		{
+			TArray<int> ElementIds;
+			UE::Geometry::FDynamicMeshColorOverlay* ColorOverlay = Mesh.Attributes()->PrimaryColors();
+			const int32 CurrentBoneIndex = Weights.Deformer.BoneNameToIndexMap[CurrentBone];
+			const int32 NumVerticesInStamp = VerticesToUpdateColor.Num();
+			for (const int32 VertexID : VerticesToUpdateColor)
+			{
+				FVector4f NewColor(GetColorOfVertex(VertexID, CurrentBoneIndex));
+				ColorOverlay->GetVertexElements(VertexID, ElementIds);
+				for (const int32 ElementId : ElementIds)
+				{
+					ColorOverlay->SetElement(ElementId, NewColor);
+				}
+				ElementIds.Reset();
+			}
+			
+		}, false);
+	PreviewMesh->NotifyDeferredEditCompleted(UPreviewMesh::ERenderUpdateMode::FastUpdate, EMeshRenderAttributeFlags::VertexColors, false);
+}
+
 float USkinWeightsPaintTool::CalculateBrushFalloff(float Distance) const
 {
 	const float f = FMathd::Clamp(1.f - BrushProperties->BrushFalloffAmount, 0.f, 1.f);
@@ -1617,7 +1704,8 @@ void USkinWeightsPaintTool::ApplyStamp(const FBrushStampData& Stamp)
 		{
 			// use mesh topology to iteratively smooth weights across neighboring vertices
 			const float UseStrength = CalculateBrushStrengthToUse(EWeightEditOperation::Relax);
-			RelaxWeightOnVertices(VerticesInStamp, VertexFalloffs, UseStrength, WeightEditsFromStamp);
+			constexpr int32 RelaxIterationsPerStamp = 3;
+			RelaxWeightOnVertices(VerticesInStamp, VertexFalloffs, RelaxIterationsPerStamp, UseStrength, WeightEditsFromStamp);
 		}
 		else
 		{
@@ -1633,46 +1721,9 @@ void USkinWeightsPaintTool::ApplyStamp(const FBrushStampData& Stamp)
 				WeightEditsFromStamp);
 		}
 	}
-	
-	{
-		TRACE_CPUPROFILER_EVENT_SCOPE(SkinTool::ApplyWeightEditsToActiveChange);
-		// store weight edits from all stamps made during a single stroke (1 transaction per stroke)
-		for (const TTuple<int32, FSingleBoneWeightEdits>& BoneWeightEdits : WeightEditsFromStamp.PerBoneWeightEdits)
-		{
-			ActiveChange->AddBoneWeightEdit(BoneWeightEdits.Value);
-		}
-	}
-	
-	{
-		TRACE_CPUPROFILER_EVENT_SCOPE(SkinTool::ApplyWeightEditsToCurrentWeights);
-		// apply weights to current weights
-		Weights.ApplyEditsToCurrentWeights(WeightEditsFromStamp);
-	}
 
-	{
-		TRACE_CPUPROFILER_EVENT_SCOPE(SkinTool::UpdateVertexColors);
-		// update vertex colors
-		PreviewMesh->DeferredEditMesh([&](FDynamicMesh3& Mesh)
-		{
-			TArray<int> ElementIds;
-			UE::Geometry::FDynamicMeshColorOverlay* ColorOverlay = Mesh.Attributes()->PrimaryColors();
-			const int32 CurrentBoneIndex = Weights.Deformer.BoneNameToIndexMap[CurrentBone];
-			const int32 NumVerticesInStamp = VerticesInStamp.Num();
-			for (int32 StampVertexIndex = 0; StampVertexIndex < NumVerticesInStamp; ++StampVertexIndex)
-			{
-				const int32 VertexID = VerticesInStamp[StampVertexIndex];
-				FVector4f NewColor(GetColorOfVertex(VertexID, CurrentBoneIndex));
-				ColorOverlay->GetVertexElements(VertexID, ElementIds);
-				for (const int32 ElementId : ElementIds)
-				{
-					ColorOverlay->SetElement(ElementId, NewColor);
-				}
-				ElementIds.Reset();
-			}
-			
-		}, false);
-		PreviewMesh->NotifyDeferredEditCompleted(UPreviewMesh::ERenderUpdateMode::FastUpdate, EMeshRenderAttributeFlags::VertexColors, false);
-	}
+	// apply weight edits to the mesh without closing the transaction
+	ApplyWeightEditsToMeshMidChange(WeightEditsFromStamp);
 }
 
 float USkinWeightsPaintTool::CalculateBrushStrengthToUse(EWeightEditOperation EditMode) const
@@ -1712,17 +1763,17 @@ float USkinWeightsPaintTool::CalculateBrushStrengthToUse(EWeightEditOperation Ed
 void USkinWeightsPaintTool::EditWeightOfBoneOnVertices(
 	EWeightEditOperation EditOperation,
 	const BoneIndex Bone,
-	const TArray<int32>& VerticesToEdit,
+	const TArray<int32>& VertexIndices,
 	const TArray<float>& VertexFalloffs,
 	const float UseStrength,
 	FMultiBoneWeightEdits& InOutWeightEdits)
 {
 	// spin through the vertices in the stamp and store new weight values in NewValuesFromStamp
 	// afterwards, these values are normalized while taking into consideration the user's desired changes
-	const int32 NumVerticesInStamp = VerticesToEdit.Num();
+	const int32 NumVerticesInStamp = VertexIndices.Num();
 	for (int32 Index = 0; Index < NumVerticesInStamp; ++Index)
 	{
-		const int32 VertexID = VerticesToEdit[Index];
+		const int32 VertexID = VertexIndices[Index];
 		const float UseFalloff = VertexFalloffs.IsValidIndex(Index) ? VertexFalloffs[Index] : 1.f;
 		const float ValueBeforeStroke = Weights.GetWeightOfBoneOnVertex(Bone, VertexID, Weights.PreChangeWeights);
 
@@ -1751,7 +1802,7 @@ void USkinWeightsPaintTool::EditWeightOfBoneOnVertices(
 			checkNoEntry();
 		}
 
-		// normalize the values across all bones affecting the vertices in the stamp, and record the bone edits
+		// normalize the values across all bones affecting this vertex, and record the bone edits
 		// normalization is done while holding all weights on the current bone constant so that user edits are not overwritten
 		Weights.EditVertexWeightAndNormalize(
 			Bone,
@@ -1762,8 +1813,9 @@ void USkinWeightsPaintTool::EditWeightOfBoneOnVertices(
 }
 
 void USkinWeightsPaintTool::RelaxWeightOnVertices(
-	TArray<int32> VerticesToEdit,
+	TArray<int32> VertexIndices,
 	TArray<float> VertexFalloffs,
+	int32 Iterations,
 	const float UseStrength,
 	FMultiBoneWeightEdits& InOutWeightEdits)
 {
@@ -1771,14 +1823,13 @@ void USkinWeightsPaintTool::RelaxWeightOnVertices(
 	{
 		return;
 	}
-
-	constexpr int32 NumRelaxIterations = 3;
-	constexpr float PercentPerIteration = 0.9f;
-	for (int32 Iteration=0; Iteration<NumRelaxIterations; ++Iteration)
+	
+	constexpr float PercentPerIteration = 0.95f;
+	for (int32 Iteration=0; Iteration<Iterations; ++Iteration)
 	{
-		for (int32 VertexIndex = 0; VertexIndex < VerticesToEdit.Num(); ++VertexIndex)
+		for (int32 VertexIndex = 0; VertexIndex < VertexIndices.Num(); ++VertexIndex)
 		{
-			const int32 VertexID = VerticesToEdit[VertexIndex];
+			const int32 VertexID = VertexIndices[VertexIndex];
 			const float UseFalloff = (VertexFalloffs.IsValidIndex(VertexIndex) ? VertexFalloffs[VertexIndex] * UseStrength : UseStrength) * PercentPerIteration;
 
 			TMap<int32, float> FinalWeights;
@@ -1864,80 +1915,52 @@ void USkinWeightsPaintTool::InitializeSmoothWeightsOperator()
 	SmoothWeightsOp->MinimumWeightThreshold = MinimumWeightThreshold;
 }
 
-void USkinWeightsPaintTool::ApplyWeightEditsToMesh(
-	const SkinPaintTool::FMultiBoneWeightEdits& WeightEdits,
-	const FText& TransactionLabel,
-	const bool bShouldTransact)
+void USkinWeightsPaintTool::ApplyWeightEditsToMeshMidChange(const SkinPaintTool::FMultiBoneWeightEdits& WeightEdits)
 {
-	if (bShouldTransact)
+	// store weight edits from all stamps made during a single stroke (1 transaction per stroke)
 	{
-		// clear the active change to start a new one
-		BeginChange();
+		TRACE_CPUPROFILER_EVENT_SCOPE(SkinTool::ApplyWeightEditsToActiveChange);
+		for (const TTuple<BoneIndex, SkinPaintTool::FSingleBoneWeightEdits>& BoneWeightEdits : WeightEdits.PerBoneWeightEdits)
+		{
+			ActiveChange->AddBoneWeightEdit(BoneWeightEdits.Value);
+		}
 	}
 
-	// store weight edits in the active change
-	for (const TTuple<int32, FSingleBoneWeightEdits>& BoneWeightEdits : WeightEdits.PerBoneWeightEdits)
+	// apply weights to current weights (triggers sparse deformation update)
 	{
-		ActiveChange->AddBoneWeightEdit(BoneWeightEdits.Value);
+		TRACE_CPUPROFILER_EVENT_SCOPE(SkinTool::ApplyWeightEditsToCurrentWeights);
+		Weights.ApplyEditsToCurrentWeights(WeightEdits);
 	}
+	
+	// queue update of vertex colors
+	WeightEdits.GetEditedVertexIndices(VerticesToUpdateColor);
+}
 
-	// apply the weight edits to the actual mesh
-	// (copies weight modifications to the tool's weight data structure and updates the vertex colors
-	// in brush mode, this is normally called by the tool itself after a stroke)
-	ActiveChange->Apply(this);
+void USkinWeightsPaintTool::ApplyWeightEditsAsTransaction(const SkinPaintTool::FMultiBoneWeightEdits& WeightEdits, const FText& TransactionLabel)
+{
+	// clear the active change to start a new one
+	BeginChange();
 
-	if (bShouldTransact)
+	// store weight edits in the active change
+	for (const TTuple<BoneIndex, FSingleBoneWeightEdits>& BoneWeightEdits : WeightEdits.PerBoneWeightEdits)
 	{
-		// store active change in the transaction buffer
-		EndChange(TransactionLabel);
+		ActiveChange->AddBoneWeightEdit(BoneWeightEdits.Value);
 	}
+	
+	// apply the weight edits of the active change to the actual mesh
+	// - copies weight modifications to the tool's weight data structure and updates the vertex colors
+	// - updates PreChangeWeights
+	ActiveChange->Apply(this);
+	
+	// store active change in the transaction buffer
+	EndChange(TransactionLabel);
 }
 
 void USkinWeightsPaintTool::UpdateCurrentBone(const FName& BoneName)
 {
 	CurrentBone = BoneName;
-	bVisibleWeightsValid = false;
-}
-
-void USkinWeightsPaintTool::GetVerticesToEdit(TArray<VertexIndex>& OutVertexIndices) const
-{
-	OutVertexIndices.Reset();
-
-	//
-	// 1. Prioritize selected vertices
-	GetSelectedVertices(OutVertexIndices);
-	if (!OutVertexIndices.IsEmpty())
-	{
-		return;
-	}
-
-	//
-	// 2. Fallback on vertices weighted to selected bones
-	if (!SelectedBoneIndices.IsEmpty())
-	{
-		VertexIndex VertexID = 0;
-		for (const VertexWeights& VertWeights : Weights.PreChangeWeights)
-		{
-			for (const FVertexBoneWeight& BoneWeight : VertWeights)
-			{
-				if (SelectedBoneIndices.Contains(BoneWeight.BoneIndex) && BoneWeight.Weight > MinimumWeightThreshold)
-				{
-					OutVertexIndices.Add(VertexID);	
-				}
-			}
-			++VertexID;
-		}
-		return;
-	}
-
-	//
-	// 3. Finally, fallback to ALL vertices in the mesh
-	const int32 NumVertices = EditedMesh->Vertices().Num();
-	OutVertexIndices.Reset(NumVertices);
-	for (VertexIndex VertexID = 0; VertexID < NumVertices; VertexID++)
-	{
-		OutVertexIndices.Add(VertexID);
-	}
+	bVertexColorsNeedUpdated = true;
+	OnSelectionChanged.Broadcast();
 }
 
 BoneIndex USkinWeightsPaintTool::GetBoneIndexFromName(const FName BoneName) const
@@ -2037,12 +2060,9 @@ void USkinWeightsPaintTool::ExternalUpdateWeights(const int32 BoneIndex, const T
 		const float Weight = Pair.Value;
 		Weights.SetWeightOfBoneOnVertex(BoneIndex, VertexID, Weight, Weights.CurrentWeights);
 		Weights.SetWeightOfBoneOnVertex(BoneIndex, VertexID, Weight, Weights.PreChangeWeights);
-	}
 
-	const FName BoneName = GetBoneNameFromIndex(BoneIndex);
-	if (BoneName == CurrentBone)
-	{
-		UpdateCurrentBoneVertexColors();
+		// queue update of vertex colors
+		VerticesToUpdateColor.Add(VertexID);
 	}
 
 	Weights.UpdateIsBoneWeighted(BoneIndex);
@@ -2220,13 +2240,12 @@ void USkinWeightsPaintTool::MirrorWeights(EAxis::Type Axis, EMirrorDirection Dir
 	const TMap<int32, int32>& VertexMirrorMap = MirrorData.GetVertexMap(); // <Target, Source>
 
 	// get set of vertices to mirror
-	TArray<VertexIndex> SelectedVertices;
-	GetVerticesToEdit(SelectedVertices);
+	TArray<VertexIndex> AllVerticesToEdit = GetVerticesToEdit();
 
 	// convert all vertex indices to the target side of the mirror plane
 	TSet<VertexIndex> VerticesToMirror;
 	
-	for (const VertexIndex SelectedVertex : SelectedVertices)
+	for (const VertexIndex SelectedVertex : AllVerticesToEdit)
 	{
 		if (VertexMirrorMap.Contains(SelectedVertex))
 		{
@@ -2275,8 +2294,7 @@ void USkinWeightsPaintTool::MirrorWeights(EAxis::Type Axis, EMirrorDirection Dir
 
 	// apply the changes
 	const FText TransactionLabel = LOCTEXT("MirrorWeightChange", "Mirror skin weights.");
-	constexpr bool bShouldTransact = true;
-	ApplyWeightEditsToMesh(WeightEditsFromMirroring, TransactionLabel, bShouldTransact);
+	ApplyWeightEditsAsTransaction(WeightEditsFromMirroring, TransactionLabel);
 
 	// warn if some vertices were not mirrored
 	if (!MirrorData.GetAllVerticesMirrored())
@@ -2285,64 +2303,53 @@ void USkinWeightsPaintTool::MirrorWeights(EAxis::Type Axis, EMirrorDirection Dir
 	}
 }
 
-void USkinWeightsPaintTool::FloodWeights(const float Weight, const EWeightEditOperation FloodMode)
-{
-	if (CurrentBone == NAME_None)
-	{
-		return;
-	}
-
-	// flood the weights on selected vertices
-	const int32 CurrentBoneIndex = Weights.Deformer.BoneNameToIndexMap[CurrentBone];
-	FMultiBoneWeightEdits WeightEditsFromFlood;
-	TArray<VertexIndex> VerticesToEdit;
-	GetVerticesToEdit(VerticesToEdit);
-	const TArray<float> VertexFalloffs = {};
-	if (FloodMode == EWeightEditOperation::Relax)
-	{
-		RelaxWeightOnVertices(VerticesToEdit, VertexFalloffs, Weight, WeightEditsFromFlood);
-	}
-	else
-	{
-		EditWeightOfBoneOnVertices(FloodMode, CurrentBoneIndex, VerticesToEdit, VertexFalloffs, Weight, WeightEditsFromFlood);	
-	}
-
-	// apply the changes
-	const FText TransactionLabel = LOCTEXT("FloodWeightChange", "Flood skin weights.");
-	constexpr bool bShouldTransact = true;
-	ApplyWeightEditsToMesh(WeightEditsFromFlood, TransactionLabel, bShouldTransact);
-}
-
-void USkinWeightsPaintTool::SetBoneWeightOnVertices(
+void USkinWeightsPaintTool::EditWeightsOnVertices(
 	BoneIndex Bone,
-	const float Weight,
-	const TArray<VertexIndex>& VerticesToEdit,
+	const float Value,
+	EWeightEditOperation EditOperation,
+	const TArray<VertexIndex>& VertexIndices,
 	const bool bShouldTransact)
 {
 	// create weight edits from setting the weight directly
 	FMultiBoneWeightEdits DirectWeightEdits;
 	const TArray<float> VertexFalloffs = {}; // no falloff
-	EditWeightOfBoneOnVertices(
-				EWeightEditOperation::Replace,
+
+	if (EditOperation == EWeightEditOperation::Relax)
+	{
+		const int32 Iterations = static_cast<int32>(Value);
+		constexpr float Strength = 1.0f; // always use max strength, and interpret value as iterations for relaxation
+		RelaxWeightOnVertices(GetVerticesToEdit(), VertexFalloffs, Iterations, Strength, DirectWeightEdits);
+	}
+	else
+	{
+		EditWeightOfBoneOnVertices(
+				EditOperation,
 				Bone,
-				VerticesToEdit,
+				VertexIndices,
 				VertexFalloffs,
-				Weight,
+				Value,
 				DirectWeightEdits);
+	}
 	
 	// apply the changes
-	const FText TransactionLabel = LOCTEXT("SetWeightChange", "Set skin weights directly.");
-	ApplyWeightEditsToMesh(DirectWeightEdits, TransactionLabel, bShouldTransact);
+	if (bShouldTransact)
+	{
+		const FText TransactionLabel = LOCTEXT("EditWeightChange", "Edit skin weights directly.");
+		ApplyWeightEditsAsTransaction(DirectWeightEdits, TransactionLabel);
+	}
+	else
+	{
+		ApplyWeightEditsToMeshMidChange(DirectWeightEdits);
+	}
 }
 
 void USkinWeightsPaintTool::PruneWeights(float Threshold)
 {
 	// set weights below the given threshold to zero
 	FMultiBoneWeightEdits WeightEditsFromPrune;
-	TArray<VertexIndex> VerticesToEdit;
-	GetVerticesToEdit(VerticesToEdit);
+	const TArray<VertexIndex>& VerticesToPrune = GetVerticesToEdit();
 	const bool bPruningAllBones = SelectedBoneIndices.IsEmpty();
-	for (const VertexIndex VertexID : VerticesToEdit)
+	for (const VertexIndex VertexID : VerticesToPrune)
 	{
 		const VertexWeights& VertexWeights = Weights.CurrentWeights[VertexID];
 		for (const FVertexBoneWeight& BoneWeight : VertexWeights)
@@ -2367,19 +2374,17 @@ void USkinWeightsPaintTool::PruneWeights(float Threshold)
 
 	// apply the changes
 	const FText TransactionLabel = LOCTEXT("PruneWeightValuesChange", "Prune skin weights.");
-	constexpr bool bShouldTransact = true;
-	ApplyWeightEditsToMesh(WeightEditsFromPrune, TransactionLabel, bShouldTransact);
+	ApplyWeightEditsAsTransaction(WeightEditsFromPrune, TransactionLabel);
 }
 
 void USkinWeightsPaintTool::AverageWeights()
 {
 	// get vertices to edit weights on
-	TArray<VertexIndex> VerticesToEdit;
-	GetVerticesToEdit(VerticesToEdit);
+	const TArray<VertexIndex>& VerticesToAverage = GetVerticesToEdit();
 
 	// accumulate ALL the weights on the vertices
 	TMap<BoneIndex, float> AccumulatedWeightMap;
-	for (const VertexIndex VertexID : VerticesToEdit)
+	for (const VertexIndex VertexID : VerticesToAverage)
 	{
 		for (const FVertexBoneWeight& BoneWeight : Weights.PreChangeWeights[VertexID])
 		{
@@ -2421,7 +2426,7 @@ void USkinWeightsPaintTool::AverageWeights()
 
 	// apply averaged weights to vertices
 	FMultiBoneWeightEdits WeightEditsFromAveraging;
-	for (const VertexIndex VertexID : VerticesToEdit)
+	for (const VertexIndex VertexID : VerticesToAverage)
 	{
 		// remove influences not a part of the average results
 		for (const FVertexBoneWeight& BoneWeight : Weights.PreChangeWeights[VertexID])
@@ -2446,17 +2451,15 @@ void USkinWeightsPaintTool::AverageWeights()
 
 	// apply the changes
 	const FText TransactionLabel = LOCTEXT("AverageWeightValuesChange", "Average skin weights.");
-	constexpr bool bShouldTransact = true;
-	ApplyWeightEditsToMesh(WeightEditsFromAveraging, TransactionLabel, bShouldTransact);
+	ApplyWeightEditsAsTransaction(WeightEditsFromAveraging, TransactionLabel);
 }
 
 void USkinWeightsPaintTool::NormalizeWeights()
 {
 	// re-set a weight on each vertex to force normalization
 	FMultiBoneWeightEdits WeightEditsFromNormalization;
-	TArray<VertexIndex> SelectedVertices;
-	GetVerticesToEdit(SelectedVertices);
-	for (const VertexIndex VertexID : SelectedVertices)
+	const TArray<VertexIndex> VerticesToNormalize = GetVerticesToEdit();
+	for (const VertexIndex VertexID : VerticesToNormalize)
 	{
 		const VertexWeights& VertexWeights = Weights.CurrentWeights[VertexID];
 		for (const FVertexBoneWeight& BoneWeight : VertexWeights)
@@ -2468,8 +2471,7 @@ void USkinWeightsPaintTool::NormalizeWeights()
 
 	// apply the changes
 	const FText TransactionLabel = LOCTEXT("NormalizeWeightValuesChange", "Normalize skin weights.");
-	constexpr bool bShouldTransact = true;
-	ApplyWeightEditsToMesh(WeightEditsFromNormalization, TransactionLabel, bShouldTransact);
+	ApplyWeightEditsAsTransaction(WeightEditsFromNormalization, TransactionLabel);
 }
 
 void USkinWeightsPaintTool::HandleSkeletalMeshModified(const TArray<FName>& InBoneNames, const ESkeletalMeshNotifyType InNotifyType)
@@ -2570,7 +2572,7 @@ void USkinWeightsPaintTool::OnActiveSkinWeightProfileChanged()
 	Weights = FSkinToolWeights();
 	Weights.Profile = WeightToolProperties->ActiveSkinWeightProfile;
 	Weights.InitializeSkinWeights(SkeletalMeshComponent, EditedMesh);
-	bVisibleWeightsValid = false;
+	bVertexColorsNeedUpdated = true;
 }
 
 bool USkinWeightsPaintTool::IsProfileValid(const FName InProfileName) const
@@ -2809,9 +2811,9 @@ void USkinWeightsPaintTool::FinishIsolatedSelection()
 	PartialMeshDescription = nullptr;
 }
 
-void USkinWeightsPaintTool::GetSelectedVertices(TArray<int32>& OutVertexIndices) const
+void USkinWeightsPaintTool::UpdateSelectedVertices()
 {
-	OutVertexIndices.Empty();
+	SelectedVertices.Empty();
 	if (!PolygonSelectionMechanic)
 	{
 		return;
@@ -2822,12 +2824,12 @@ void USkinWeightsPaintTool::GetSelectedVertices(TArray<int32>& OutVertexIndices)
 	const FVertexArray& AllVertices = EditedMesh->Vertices();
 
 	// validate and add vertices to the output array
-	auto AddVertices = [&OutVertexIndices, &AllVertices](const TSet<int32>& VerticesToAdd)
+	auto AddVertices = [this, &AllVertices](const TSet<int32>& VerticesToAdd)
 	{
 		// we have to make sure that the vertex ids are safe to use as PolygonSelectionMechanic does not act on the
 		// mesh description but on the dynamic mesh that can duplicate vertices when dealing with degenerate triangles.
 		// cf. FMeshDescriptionToDynamicMesh::Convert for more details.
-		Algo::CopyIf(VerticesToAdd, OutVertexIndices, [&](int32 VertexID)
+		Algo::CopyIf(VerticesToAdd, SelectedVertices, [&](int32 VertexID)
 		{
 			return AllVertices.IsValid(VertexID);	
 		});
@@ -2864,6 +2866,40 @@ void USkinWeightsPaintTool::GetSelectedVertices(TArray<int32>& OutVertexIndices)
 	}
 }
 
+const TArray<int32>& USkinWeightsPaintTool::GetSelectedVertices() const
+{
+	return SelectedVertices;
+}
+
+const TArray<int32>& USkinWeightsPaintTool::GetVerticesToEdit()
+{
+	// prioritize re-using selected vertices array (if there are any)
+	if (!SelectedVertices.IsEmpty())
+	{
+		return SelectedVertices;
+	}
+	
+	// fallback on vertices weighted to selected bones
+	VerticesToEdit.Reset();
+	if (!SelectedBoneIndices.IsEmpty())
+	{
+		VertexIndex VertexID = 0;
+		for (const VertexWeights& VertWeights : Weights.PreChangeWeights)
+		{
+			for (const FVertexBoneWeight& BoneWeight : VertWeights)
+			{
+				if (SelectedBoneIndices.Contains(BoneWeight.BoneIndex) && BoneWeight.Weight > MinimumWeightThreshold)
+				{
+					VerticesToEdit.Add(VertexID);	
+				}
+			}
+			++VertexID;
+		}
+	}
+
+	return VerticesToEdit;
+}
+
 void USkinWeightsPaintTool::GetSelectedTriangles(TArray<int32>& OutTriangleIndices) const
 {
 	OutTriangleIndices.Empty();
@@ -2955,12 +2991,19 @@ FName USkinWeightsPaintTool::GetBoneNameFromIndex(BoneIndex InIndex) const
 	return NAME_None;
 }
 
-void USkinWeightsPaintTool::SetDisplayVertexColors(bool bShowVertexColors) const
+
+BoneIndex USkinWeightsPaintTool::GetCurrentBoneIndex() const
+{
+	return GetBoneIndexFromName(CurrentBone);
+}
+
+void USkinWeightsPaintTool::SetDisplayVertexColors(bool bShowVertexColors)
 {
 	if (bShowVertexColors)
 	{
 		UMaterialInterface* VtxColorMaterial = GetToolManager()->GetContextQueriesAPI()->GetStandardMaterial(EStandardToolContextMaterials::VertexColorMaterial);
 		PreviewMesh->SetOverrideRenderMaterial(VtxColorMaterial);
+		bVertexColorsNeedUpdated = true;
 	}
 	else
 	{
@@ -2999,7 +3042,7 @@ void USkinWeightsPaintTool::OnPropertyModified(UObject* ModifiedObject, FPropert
 		GET_MEMBER_NAME_STRING_CHECKED(FLinearColor, A)};
 	if (ColorPropertyNames.Contains(NameOfModifiedProperty))
 	{
-		bVisibleWeightsValid = false;
+		bVertexColorsNeedUpdated = true;
 
 		// force all colors to have Alpha = 1
 		WeightToolProperties->MinColor.A = 1.f;
diff --git a/Engine/Plugins/Experimental/MeshModelingToolsetExp/Source/MeshModelingToolsEditorOnlyExp/Public/SkeletalMesh/SkinWeightsPaintTool.h b/Engine/Plugins/Experimental/MeshModelingToolsetExp/Source/MeshModelingToolsEditorOnlyExp/Public/SkeletalMesh/SkinWeightsPaintTool.h
--- a/Engine/Plugins/Experimental/MeshModelingToolsetExp/Source/MeshModelingToolsEditorOnlyExp/Public/SkeletalMesh/SkinWeightsPaintTool.h
+++ b/Engine/Plugins/Experimental/MeshModelingToolsetExp/Source/MeshModelingToolsEditorOnlyExp/Public/SkeletalMesh/SkinWeightsPaintTool.h
@@ -97,10 +97,10 @@ namespace SkinPaintTool
 	struct FVertexBoneWeight
 	{
 		FVertexBoneWeight() : BoneIndex(INDEX_NONE), VertexInBoneSpace(FVector::ZeroVector), Weight(0.0f) {}
-		FVertexBoneWeight(int32 InBoneIndex, const FVector& InPosInRefPose, float InWeight) :
+		FVertexBoneWeight(BoneIndex InBoneIndex, const FVector& InPosInRefPose, float InWeight) :
 			BoneIndex(InBoneIndex), VertexInBoneSpace(InPosInRefPose), Weight(InWeight){}
 		
-		int32 BoneIndex;
+		BoneIndex BoneIndex;
 		FVector VertexInBoneSpace;
 		float Weight;
 	};
@@ -152,6 +152,7 @@ namespace SkinPaintTool
 		void MergeSingleEdit(const int32 BoneIndex, const int32 VertexID, const float OldWeight, const float NewWeight);
 		void MergeEdits(const FSingleBoneWeightEdits& BoneWeightEdits);
 		float GetVertexDeltaFromEdits(const int32 BoneIndex, const int32 VertexIndex);
+		void GetEditedVertexIndices(TSet<int32>& OutVerticesToEdit) const;
 
 		// map of bone indices to weight edits made to that bone
 		TMap<BoneIndex, FSingleBoneWeightEdits> PerBoneWeightEdits;
@@ -295,6 +296,19 @@ struct FSkinWeightBrushConfig
 	EWeightBrushFalloffMode FalloffMode = EWeightBrushFalloffMode::Surface;
 };
 
+struct MESHMODELINGTOOLSEDITORONLYEXP_API FDirectEditWeightState
+{
+	EWeightEditOperation EditMode;
+	float StartValue = 0.f;
+	float CurrentValue = 0.f;
+	bool bInTransaction = false;
+
+	void Reset();
+	float GetModeDefaultValue();
+	float GetModeMinValue();
+	float GetModeMaxValue();
+};
+
 // Container for properties displayed in Details panel while using USkinWeightsPaintTool
 UCLASS(config = EditorSettings)
 class MESHMODELINGTOOLSEDITORONLYEXP_API USkinWeightsPaintToolProperties : public UBrushBaseProperties
@@ -327,7 +341,6 @@ public:
 	FLinearColor MinColor;
 	UPROPERTY(EditAnywhere, Config, Category = MeshDisplay)
 	FLinearColor MaxColor;
-	bool bColorModeChanged = false;
 
 	// weight editing arguments
 	UPROPERTY(Config)
@@ -335,9 +348,13 @@ public:
 	UPROPERTY(Config)
 	EMirrorDirection MirrorDirection = EMirrorDirection::PositiveToNegative;
 	UPROPERTY(Config)
-	float FloodValue = 1.f;
-	UPROPERTY(Config)
 	float PruneValue = 0.01;
+	UPROPERTY(Config)
+	float ReplaceValue = 1.0;
+	UPROPERTY(Config)
+	int32 RelaxIterations = 5;
+	// the state of the direct weight editing tools (mode buttons + slider)
+	FDirectEditWeightState DirectEditState;
 
 	// save/restore user specified settings for each tool mode
 	FSkinWeightBrushConfig& GetBrushConfig();
@@ -413,15 +430,15 @@ public:
 
 	// weight editing operations (selection based)
 	void MirrorWeights(EAxis::Type Axis, EMirrorDirection Direction);
-	void FloodWeights(const float Weight, const EWeightEditOperation FloodMode);
 	void PruneWeights(const float Threshold);
 	void AverageWeights();
 	void NormalizeWeights();
 	
 	// method to set weights directly (numeric input, for example)
-	void SetBoneWeightOnVertices(
+	void EditWeightsOnVertices(
 		BoneIndex Bone,
-		const float Weight,
+		const float Value,
+		EWeightEditOperation EditOperation,
 		const TArray<VertexIndex>& VerticesToEdit,
 		const bool bShouldTransact);
 
@@ -438,42 +455,42 @@ public:
 	bool IsSelectionIsolated() const;
 	void SetIsolateSelected(const bool bIsolateSelection);
 
-	// get a list of currently selected components (converting other components)
-	void GetSelectedVertices(TArray<int32>& OutVertexIndices) const;
+	// get a list of currently selected vertices (converting edges and faces to vertices)
+	const TArray<int32>& GetSelectedVertices() const;
+	const TArray<int32>& GetVerticesToEdit();
 	void GetSelectedTriangles(TArray<int32>& OutTriangleIndices) const;
 
 	// get the average weight value of each influence on the given vertices
 	void GetInfluences(const TArray<int32>& VertexIndices, TArray<BoneIndex>& OutBoneIndices);
-
 	// get the average weight value of a single bone on the given vertices
 	float GetAverageWeightOnBone(const BoneIndex InBoneIndex, const TArray<int32>& VertexIndices);
-
 	// convert an index to a name
 	FName GetBoneNameFromIndex(BoneIndex InIndex) const;
+	// get the currently selected bone
+	BoneIndex GetCurrentBoneIndex() const;
 
 	// toggle the display of weights on the preview mesh (if false, uses the normal skeletal mesh material)
-	void SetDisplayVertexColors(bool bShowVertexColors=true) const;
+	void SetDisplayVertexColors(bool bShowVertexColors=true);
 	// set focus back to viewport so that hotkeys are immediately detected while hovering
 	void SetFocusInViewport() const;
 
 	// HOW TO EDIT WEIGHTS WITH UNDO/REDO:
 	//
-	// Live Edits:
+	// "Interactive" Edits:
 	// For multiple weight editing operations that need to be grouped into a single transaction, like dragging a slider or
-	// dragging a brush, you must call BeginChange(), then ApplyWeightsEditsToMesh(bShouldTransact=false), then EndChange().
-	// All the edits are stored into the "ActiveChange" and applied as a single transaction. Deformations and vertex
-	// colors will be updated throughout the duration of the change.
-	//
-	// Single Edits:
-	// For all one-and-done edits, you can call ApplyWeightEditsToMesh(bShouldTransact=True).
-	// It will Begin/End the change and create a transaction for it.
-	//
+	// dragging a brush, you must call:
+	//  1. BeginChange()
+	//  2. ApplyWeightEditsToMeshMidChange() (this may be called multiple times)
+	//  2. EndChange()
+	// All the edits are stored into the "ActiveChange" and applied as a single transaction in EndChange().
+	// Deformations and vertex colors will be updated throughout the duration of the change.
 	void BeginChange();
 	void EndChange(const FText& TransactionLabel);
-	void ApplyWeightEditsToMesh(
-		const SkinPaintTool::FMultiBoneWeightEdits& WeightEdits,
-		const FText& TransactionLabel,
-		const bool bShouldTransact);
+	void ApplyWeightEditsToMeshMidChange(const SkinPaintTool::FMultiBoneWeightEdits& WeightEdits);
+	// "One-off" Edits:
+	// For all one-and-done edits, you can call ApplyWeightEditsAsTransaction().
+	// It will Begin/End the change and create a transaction for it.
+	void ApplyWeightEditsAsTransaction(const SkinPaintTool::FMultiBoneWeightEdits& WeightEdits, const FText& TransactionLabel);
 
 	// called whenever the selection is modified
 	DECLARE_MULTICAST_DELEGATE(FOnSelectionChanged);
@@ -523,6 +540,7 @@ protected:
 	void RelaxWeightOnVertices(
 		TArray<int32> VerticesToEdit,
 		TArray<float> VertexFalloffs,
+		int32 Iterations,
 		const float UseStrength,
 		SkinPaintTool::FMultiBoneWeightEdits& InOutWeightEdits);
 
@@ -566,10 +584,14 @@ protected:
 	TUniquePtr<UE::Geometry::TSmoothBoneWeights<int32, float>> SmoothWeightsOp;
 	void InitializeSmoothWeightsOperator();
 
-	// vertex colors updated when switching current bone or editing weights
-	void UpdateCurrentBoneVertexColors();
+	// vertex colors updated when switching current bone or initializing whole mesh
+	void UpdateVertexColorForAllVertices();
+	bool bVertexColorsNeedUpdated = false;
+	// vertex colors updated when make sparse edits to subset of vertices
+	void UpdateVertexColorForSubsetOfVertices();
+	TSet<int32> VerticesToUpdateColor;
+	
 	FVector4f GetColorOfVertex(VertexIndex InVertexIndex, BoneIndex InBoneIndex) const;
-	bool bVisibleWeightsValid = false;
 
 	// which bone are we currently painting?
 	void UpdateCurrentBone(const FName &BoneName);
@@ -579,7 +601,7 @@ protected:
 	TArray<BoneIndex> SelectedBoneIndices;
 
 	// determines the set of vertices to operate on, using selection as the priority
-	void GetVerticesToEdit(TArray<VertexIndex>& OutVertexIndices) const;
+	void UpdateSelectedVertices();
 	
 	BoneIndex GetBoneIndexFromName(const FName BoneName) const;
 
@@ -592,6 +614,8 @@ protected:
 	TUniquePtr<UE::Geometry::FDynamicMeshAABBTree3> MeshSpatial = nullptr;
 	TUniquePtr<UE::Geometry::FTriangleGroupTopology> SelectionTopology = nullptr;
 	void InitializeSelectionMechanic();
+	TArray<VertexIndex> SelectedVertices;
+	TArray<VertexIndex> VerticesToEdit;
 
 	// isolate selection sub-meshes
 	UE::Geometry::FDynamicSubmesh3 PartialSubMesh;
