Fixing localization conflicts

diff --git a/Engine/Plugins/Experimental/Animation/SkeletalMeshModelingTools/Source/SkeletalMeshModelingTools/Private/DetailCustomization/SkinWeightDetailCustomization.cpp b/Engine/Plugins/Experimental/Animation/SkeletalMeshModelingTools/Source/SkeletalMeshModelingTools/Private/DetailCustomization/SkinWeightDetailCustomization.cpp
--- a/Engine/Plugins/Experimental/Animation/SkeletalMeshModelingTools/Source/SkeletalMeshModelingTools/Private/DetailCustomization/SkinWeightDetailCustomization.cpp
+++ b/Engine/Plugins/Experimental/Animation/SkeletalMeshModelingTools/Source/SkeletalMeshModelingTools/Private/DetailCustomization/SkinWeightDetailCustomization.cpp
@@ -709,7 +709,7 @@ void FSkinWeightDetailCustomization::AddSelectionUI(IDetailLayoutBuilder& Detail
 				SNew(STextBlock)
 				.Text(LOCTEXT("RelaxIterationsLabel", "Iterations"))
 				.Font(FAppStyle::GetFontStyle(TEXT("PropertyWindow.NormalFont")))
-				.ToolTipText(LOCTEXT("PruneThresholdTooltip", "The number of iterations of relaxation to apply."))
+				.ToolTipText(LOCTEXT("RelaxIterationsTooltip", "The number of iterations of relaxation to apply."))
 			]
 
 			+SHorizontalBox::Slot()
@@ -738,7 +738,7 @@ void FSkinWeightDetailCustomization::AddSelectionUI(IDetailLayoutBuilder& Detail
 					SNew(SButton)
 					.HAlign(HAlign_Center)
 					.Text(LOCTEXT("RelaxWeightsButtonLabel", "Relax"))
-					.ToolTipText(LOCTEXT("PruneButtonTooltip",
+					.ToolTipText(LOCTEXT("RelaxButtonTooltip",
 						"Relax: the weight of each vertex is replaced by the average of it's neighbors. This smooths weights across the mesh.\n"))
 					.OnClicked_Lambda([this]()
 					{
