Fixes issue, skeletal mesh editing tool edits are no longer auto-accepted when switching between tools

diff --git a/Engine/Plugins/Experimental/Animation/SkeletalMeshModelingTools/Source/SkeletalMeshModelingTools/Private/SkeletalMeshModelingToolsEditorMode.cpp b/Engine/Plugins/Experimental/Animation/SkeletalMeshModelingTools/Source/SkeletalMeshModelingTools/Private/SkeletalMeshModelingToolsEditorMode.cpp
--- a/Engine/Plugins/Experimental/Animation/SkeletalMeshModelingTools/Source/SkeletalMeshModelingTools/Private/SkeletalMeshModelingToolsEditorMode.cpp
+++ b/Engine/Plugins/Experimental/Animation/SkeletalMeshModelingTools/Source/SkeletalMeshModelingTools/Private/SkeletalMeshModelingToolsEditorMode.cpp
@@ -248,8 +248,14 @@ void USkeletalMeshModelingToolsEditorMode::Enter()
 
 	// Skeleton Editing
 	RegisterTool(ToolManagerCommands.BeginSkeletonEditingTool, TEXT("BeginSkeletonEditingTool"), NewObject<USkeletonEditingToolBuilder>());
-	
+
+	// highlights skin weights tool by default
 	GetInteractiveToolsContext()->ToolManager->SelectActiveToolType(EToolSide::Left, TEXT("BeginSkinWeightsPaintTool"));
+
+	// record switching behavior to restore on exit
+	ToolSwitchModeToRestoreOnExit = GetInteractiveToolsContext()->ToolManager->GetToolSwitchMode();
+	// default to NOT applying changes to skeletal meshes when switching between tools without accepting
+	GetInteractiveToolsContext()->ToolManager->SetToolSwitchMode(EToolManagerToolSwitchMode::CancelIfAble);
 }
 
 UDebugSkelMeshComponent* USkeletalMeshModelingToolsEditorMode::GetSkelMeshComponent() const
@@ -269,6 +275,9 @@ void USkeletalMeshModelingToolsEditorMode::Exit()
 	UE::TransformGizmoUtil::DeregisterTransformGizmoContextObject(InteractiveToolsContext);
 	UE::SkeletalMeshGizmoUtils::UnregisterTransformGizmoContextObject(InteractiveToolsContext);
 	UE::SkeletalMeshEditorUtils::UnregisterEditorContextObject(InteractiveToolsContext);
+
+	// restore previous tool switching behavior
+	GetInteractiveToolsContext()->ToolManager->SetToolSwitchMode(ToolSwitchModeToRestoreOnExit);
 	
 #if ENABLE_STYLUS_SUPPORT
 	StylusStateTracker = nullptr;
diff --git a/Engine/Plugins/Experimental/Animation/SkeletalMeshModelingTools/Source/SkeletalMeshModelingTools/Private/SkeletalMeshModelingToolsEditorMode.h b/Engine/Plugins/Experimental/Animation/SkeletalMeshModelingTools/Source/SkeletalMeshModelingTools/Private/SkeletalMeshModelingToolsEditorMode.h
--- a/Engine/Plugins/Experimental/Animation/SkeletalMeshModelingTools/Source/SkeletalMeshModelingTools/Private/SkeletalMeshModelingToolsEditorMode.h
+++ b/Engine/Plugins/Experimental/Animation/SkeletalMeshModelingTools/Source/SkeletalMeshModelingTools/Private/SkeletalMeshModelingToolsEditorMode.h
@@ -16,6 +16,7 @@ class ISkeletalMeshEditingInterface;
 class HHitProxy;
 class UDebugSkelMeshComponent;
 class ISkeletalMeshEditor;
+enum class EToolManagerToolSwitchMode;
 
 UCLASS()
 class USkeletalMeshModelingToolsEditorMode : 
@@ -53,6 +54,9 @@ private:
 	// Stylus support is currently disabled; this is left in for reference if/when it is brought back
 	//TUniquePtr<FStylusStateTracker> StylusStateTracker;
 
+	// we restore previous switch tool behavior when exiting this mode
+	EToolManagerToolSwitchMode ToolSwitchModeToRestoreOnExit;
+
 	static ISkeletalMeshEditingInterface* GetSkeletonInterface(UInteractiveTool* InTool);
 
 	UDebugSkelMeshComponent* GetSkelMeshComponent() const;
