Added option for UEditorInteractiveToolsContext to stay active during a PIE session.

diff --git a/Engine/Plugins/Experimental/Animation/SkeletalMeshModelingTools/Source/SkeletalMeshModelingTools/Private/SkeletalMeshModelingToolsEditorMode.cpp b/Engine/Plugins/Experimental/Animation/SkeletalMeshModelingTools/Source/SkeletalMeshModelingTools/Private/SkeletalMeshModelingToolsEditorMode.cpp
--- a/Engine/Plugins/Experimental/Animation/SkeletalMeshModelingTools/Source/SkeletalMeshModelingTools/Private/SkeletalMeshModelingToolsEditorMode.cpp
+++ b/Engine/Plugins/Experimental/Animation/SkeletalMeshModelingTools/Source/SkeletalMeshModelingTools/Private/SkeletalMeshModelingToolsEditorMode.cpp
@@ -186,6 +186,10 @@ void USkeletalMeshModelingToolsEditorMode::Enter()
 {
 	UEdMode::Enter();
 
+	UEditorInteractiveToolsContext* EditorInteractiveToolsContext = GetInteractiveToolsContext(EToolsContextScope::Editor);
+	bDeactivateOnPIEStartStateToRestore = EditorInteractiveToolsContext->GetDeactivateToolsOnPIEStart();
+	EditorInteractiveToolsContext->SetDeactivateToolsOnPIEStart(false);
+
 	UEditorInteractiveToolsContext* InteractiveToolsContext = GetInteractiveToolsContext();
 
 	if (TObjectPtr<UToolTargetManager> ToolTargetManager = InteractiveToolsContext->TargetManager)
@@ -275,6 +279,9 @@ void USkeletalMeshModelingToolsEditorMode::Exit()
 	UE::TransformGizmoUtil::DeregisterTransformGizmoContextObject(InteractiveToolsContext);
 	UE::SkeletalMeshGizmoUtils::UnregisterTransformGizmoContextObject(InteractiveToolsContext);
 	UE::SkeletalMeshEditorUtils::UnregisterEditorContextObject(InteractiveToolsContext);
+	
+	UEditorInteractiveToolsContext* EditorInteractiveToolsContext = GetInteractiveToolsContext(EToolsContextScope::Editor);
+	EditorInteractiveToolsContext->SetDeactivateToolsOnPIEStart(bDeactivateOnPIEStartStateToRestore);
 
 	// restore previous tool switching behavior
 	GetInteractiveToolsContext()->ToolManager->SetToolSwitchMode(ToolSwitchModeToRestoreOnExit);
@@ -409,6 +416,13 @@ void USkeletalMeshModelingToolsEditorMode::OnToolEnded(UInteractiveToolManager*
 	}
 }
 
+bool USkeletalMeshModelingToolsEditorMode::ShouldToolStartBeAllowed(const FString& ToolIdentifier) const
+{
+	// in the base mode, this returns false if the level editor is in PIE or simulated
+	// we allow all skeletal mesh editing tools to be started while running in PIE / simulate
+	return true;
+}
+
 void USkeletalMeshModelingToolsEditorMode::SetEditorBinding(const TWeakPtr<ISkeletalMeshEditor>& InSkeletalMeshEditor)
 {
 	if (!InSkeletalMeshEditor.IsValid())
diff --git a/Engine/Plugins/Experimental/Animation/SkeletalMeshModelingTools/Source/SkeletalMeshModelingTools/Private/SkeletalMeshModelingToolsEditorMode.h b/Engine/Plugins/Experimental/Animation/SkeletalMeshModelingTools/Source/SkeletalMeshModelingTools/Private/SkeletalMeshModelingToolsEditorMode.h
--- a/Engine/Plugins/Experimental/Animation/SkeletalMeshModelingTools/Source/SkeletalMeshModelingTools/Private/SkeletalMeshModelingToolsEditorMode.h
+++ b/Engine/Plugins/Experimental/Animation/SkeletalMeshModelingTools/Source/SkeletalMeshModelingTools/Private/SkeletalMeshModelingToolsEditorMode.h
@@ -42,6 +42,7 @@ public:
 	virtual bool HandleClick(FEditorViewportClient* InViewportClient, HHitProxy *HitProxy, const FViewportClick &Click) override;
 	virtual bool ComputeBoundingBoxForViewportFocus(AActor* Actor, UPrimitiveComponent* PrimitiveComponent, FBox& InOutBox) const override;
 	virtual bool UsesToolkits() const override { return true; }
+	virtual bool ShouldToolStartBeAllowed(const FString& ToolIdentifier) const override;
 
 	// binding
 	void SetEditorBinding(const TWeakPtr<ISkeletalMeshEditor>& InSkeletalMeshEditor);
@@ -63,6 +64,8 @@ private:
 
 	bool NeedsTransformGizmo() const;
 
+	bool bDeactivateOnPIEStartStateToRestore;
+
 	FDelegateHandle ToToolNotifierHandle;
 	FDelegateHandle FromToolNotifierHandle;
 	
diff --git a/Engine/Source/Editor/Experimental/EditorInteractiveToolsFramework/Private/EdModeInteractiveToolsContext.cpp b/Engine/Source/Editor/Experimental/EditorInteractiveToolsFramework/Private/EdModeInteractiveToolsContext.cpp
--- a/Engine/Source/Editor/Experimental/EditorInteractiveToolsFramework/Private/EdModeInteractiveToolsContext.cpp
+++ b/Engine/Source/Editor/Experimental/EditorInteractiveToolsFramework/Private/EdModeInteractiveToolsContext.cpp
@@ -396,7 +396,10 @@ void UEditorInteractiveToolsContext::ShutdownContext()
 
 void UEditorInteractiveToolsContext::TerminateActiveToolsOnPIEStart()
 {
-	DeactivateAllActiveTools(EToolShutdownType::Accept);
+	if (bDeactivateOnPIEStart)
+	{
+		DeactivateAllActiveTools(EToolShutdownType::Accept);	
+	}
 }
 void UEditorInteractiveToolsContext::TerminateActiveToolsOnSaveWorld()
 {
@@ -848,6 +851,11 @@ void UEditorInteractiveToolsContext::SetAbsoluteWorldSnappingEnabled(bool bEnabl
 	bEnableAbsoluteWorldSnapping = bEnabled;
 }
 
+void UEditorInteractiveToolsContext::SetDeactivateToolsOnPIEStart(bool bDeactivateTools)
+{
+	bDeactivateOnPIEStart = bDeactivateTools;
+}
+
 void UModeManagerInteractiveToolsContext::Tick(FEditorViewportClient* ViewportClient, float DeltaTime)
 {
 	UEditorInteractiveToolsContext::Tick(ViewportClient, DeltaTime);
diff --git a/Engine/Source/Editor/Experimental/EditorInteractiveToolsFramework/Public/EdModeInteractiveToolsContext.h b/Engine/Source/Editor/Experimental/EditorInteractiveToolsFramework/Public/EdModeInteractiveToolsContext.h
--- a/Engine/Source/Editor/Experimental/EditorInteractiveToolsFramework/Public/EdModeInteractiveToolsContext.h
+++ b/Engine/Source/Editor/Experimental/EditorInteractiveToolsFramework/Public/EdModeInteractiveToolsContext.h
@@ -172,6 +172,17 @@ public:
 	/** @return true if Absolute World Snapping mode is Enabled */
 	bool GetAbsoluteWorldSnappingEnabled() const { return bEnableAbsoluteWorldSnapping; }
 
+	/*
+	 * Configure whether tools should shutdown when entering PIE.
+	 * By default, the context will shut down any active tools on PIE start.
+	 * Setting this to true will prevent this, but the system makes no promises about whether tools work properly across PIE start/run/shutdown.
+	 * Therefore, do not use this if you do not know for certain that it is safe to persist tools across PIE start in your situation.
+	 */
+	void SetDeactivateToolsOnPIEStart(bool bDeactivateTools);
+
+	/** @return true if tools should be deactivated when PIE is started */
+	bool GetDeactivateToolsOnPIEStart() const { return bDeactivateOnPIEStart; }
+
 protected:
 
 	// we hide these 
@@ -228,6 +239,8 @@ private:
 	bool bForceCombinedGizmoMode = false;
 	bool bEnableAbsoluteWorldSnapping = false;
 
+	bool bDeactivateOnPIEStart = true;
+
 	bool bIsActive = false;
 };
 
