Skin weight tool: fixes brush properties showing up twice in details panel.

diff --git a/Engine/Plugins/Experimental/Animation/SkeletalMeshModelingTools/Source/SkeletalMeshModelingTools/Private/DetailCustomization/SkinWeightDetailCustomization.cpp b/Engine/Plugins/Experimental/Animation/SkeletalMeshModelingTools/Source/SkeletalMeshModelingTools/Private/DetailCustomization/SkinWeightDetailCustomization.cpp
--- a/Engine/Plugins/Experimental/Animation/SkeletalMeshModelingTools/Source/SkeletalMeshModelingTools/Private/DetailCustomization/SkinWeightDetailCustomization.cpp
+++ b/Engine/Plugins/Experimental/Animation/SkeletalMeshModelingTools/Source/SkeletalMeshModelingTools/Private/DetailCustomization/SkinWeightDetailCustomization.cpp
@@ -140,15 +140,11 @@ void FSkinWeightDetailCustomization::CustomizeDetails(IDetailLayoutBuilder& Deta
 	// Edit SkinWeightLayer category 
 	IDetailCategoryBuilder& SkinWeightLayerCategory = DetailBuilder.EditCategory("SkinWeightLayer", FText::GetEmpty(), ECategoryPriority::Important);
 	SkinWeightLayerCategory.InitiallyCollapsed(true);
-
-	// hide all properties that have been customized
+	
+	// hide skin weight tool properties that were customized
 	TArray<FName> PropertiesToHide;
 	PropertiesToHide.Add(GET_MEMBER_NAME_CHECKED(USkinWeightsPaintToolProperties, BrushMode));
-	PropertiesToHide.Add(GET_MEMBER_NAME_CHECKED(UBrushBaseProperties, BrushSize));
-	PropertiesToHide.Add(GET_MEMBER_NAME_CHECKED(UBrushBaseProperties, BrushStrength));
-	PropertiesToHide.Add(GET_MEMBER_NAME_CHECKED(UBrushBaseProperties, BrushFalloffAmount));
-	PropertiesToHide.Add(GET_MEMBER_NAME_CHECKED(UBrushBaseProperties, BrushRadius));
-	PropertiesToHide.Add(GET_MEMBER_NAME_CHECKED(UBrushBaseProperties, bSpecifyRadius));
+	
 	PropertiesToHide.Add(GET_MEMBER_NAME_CHECKED(USkinWeightsPaintToolProperties, EditingMode));
 	PropertiesToHide.Add(GET_MEMBER_NAME_CHECKED(USkinWeightsPaintToolProperties, ColorMode));
 	PropertiesToHide.Add(GET_MEMBER_NAME_CHECKED(USkinWeightsPaintToolProperties, MeshSelectMode));
@@ -160,6 +156,18 @@ void FSkinWeightDetailCustomization::CustomizeDetails(IDetailLayoutBuilder& Deta
 		TSharedRef<IPropertyHandle> Property = DetailBuilder.GetProperty(PropertyToHide);
 		DetailBuilder.HideProperty(Property);
 	}
+	// hide base class properties
+	PropertiesToHide.Reset();
+	PropertiesToHide.Add(GET_MEMBER_NAME_CHECKED(UBrushBaseProperties, bSpecifyRadius));
+	PropertiesToHide.Add(GET_MEMBER_NAME_CHECKED(UBrushBaseProperties, BrushSize));
+	PropertiesToHide.Add(GET_MEMBER_NAME_CHECKED(UBrushBaseProperties, BrushStrength));
+	PropertiesToHide.Add(GET_MEMBER_NAME_CHECKED(UBrushBaseProperties, BrushFalloffAmount));
+	PropertiesToHide.Add(GET_MEMBER_NAME_CHECKED(UBrushBaseProperties, BrushRadius));
+	for (const FName PropertyToHide : PropertiesToHide)
+	{
+		TSharedRef<IPropertyHandle> Property = DetailBuilder.GetProperty(PropertyToHide, UBrushBaseProperties::StaticClass());
+		DetailBuilder.HideProperty(Property);
+	}
 }
 
 void FSkinWeightDetailCustomization::AddBrushUI(IDetailLayoutBuilder& DetailBuilder) const
