Skeletal Mesh Editing: Mark the plugin as being release-ready.

diff --git a/Engine/Plugins/Experimental/Animation/SkeletalMeshModelingTools/SkeletalMeshModelingTools.uplugin b/Engine/Plugins/Experimental/Animation/SkeletalMeshModelingTools/SkeletalMeshModelingTools.uplugin
--- a/Engine/Plugins/Experimental/Animation/SkeletalMeshModelingTools/SkeletalMeshModelingTools.uplugin
+++ b/Engine/Plugins/Experimental/Animation/SkeletalMeshModelingTools/SkeletalMeshModelingTools.uplugin
@@ -1,7 +1,7 @@
 ﻿{
 	"FileVersion": 3,
 	"Version": 1,
-	"VersionName": "0.1",
+	"VersionName": "1.0",
 	"FriendlyName": "Skeletal Mesh Editing Tools",
 	"Description": "Create skeletons, paint skin weights and edit skeletal meshes.",
 	"Category" : "Animation",
@@ -9,10 +9,10 @@
 	"CreatedByURL": "https://epicgames.com/",
 	"DocsURL": "",
 	"MarketplaceURL": "",
-	"EnabledByDefault" : false,
+	"EnabledByDefault" : true,
 	"CanContainContent": true,
 	"IsExperimentalVersion": false,
-	"IsBetaVersion": true,
+	"IsBetaVersion": false,
 	"Installed": false,
 	"Modules": [
 		{
