Fixing StaticMesh asset context menus to respect read-only assets
#rnx

diff --git a/Engine/Plugins/Editor/EngineAssetDefinitions/Source/Private/AssetDefinition_StaticMesh.cpp b/Engine/Plugins/Editor/EngineAssetDefinitions/Source/Private/AssetDefinition_StaticMesh.cpp
--- a/Engine/Plugins/Editor/EngineAssetDefinitions/Source/Private/AssetDefinition_StaticMesh.cpp
+++ b/Engine/Plugins/Editor/EngineAssetDefinitions/Source/Private/AssetDefinition_StaticMesh.cpp
@@ -506,7 +506,8 @@ namespace MenuExtension_StaticMesh
 			FToolMenuSection& Section = Menu->FindOrAddSection("GetAssetActions");
 			Section.AddDynamicEntry(NAME_None, FNewToolMenuSectionDelegate::CreateLambda([](FToolMenuSection& InSection)
 			{
-				if (const UContentBrowserAssetContextMenuContext* Context = UContentBrowserAssetContextMenuContext::FindContextWithAssets(InSection))
+				const UContentBrowserAssetContextMenuContext* Context = UContentBrowserAssetContextMenuContext::FindContextWithAssets(InSection);
+				if (Context && Context->bCanBeModified)
 				{
 					if (CVarEnableSaveGeneratedLODsInPackage.GetValueOnGameThread() != 0)
 					{
diff --git a/Engine/Plugins/Experimental/Animation/SkeletalMeshModelingTools/Source/SkeletalMeshModelingTools/Private/SkeletalMeshModelingToolsModule.cpp b/Engine/Plugins/Experimental/Animation/SkeletalMeshModelingTools/Source/SkeletalMeshModelingTools/Private/SkeletalMeshModelingToolsModule.cpp
--- a/Engine/Plugins/Experimental/Animation/SkeletalMeshModelingTools/Source/SkeletalMeshModelingTools/Private/SkeletalMeshModelingToolsModule.cpp
+++ b/Engine/Plugins/Experimental/Animation/SkeletalMeshModelingTools/Source/SkeletalMeshModelingTools/Private/SkeletalMeshModelingToolsModule.cpp
@@ -113,19 +113,23 @@ void FSkeletalMeshModelingToolsModule::RegisterMenusAndToolbars()
 		FToolMenuSection& Section = ToolMenu->FindOrAddSection("GetAssetActions");
 		Section.AddDynamicEntry("ConvertToSkeletalMesh", FNewToolMenuSectionDelegate::CreateLambda([](FToolMenuSection& InSection)
 		{
-			const TAttribute<FText> Label = LOCTEXT("ConvertToSkeletalMesh", "Convert to Skeletal Mesh");
-			const TAttribute<FText> ToolTip = LOCTEXT("ConvertToSkeletalMeshToolTip", "Converts a non-Nanite static mesh to a skeletal mesh with an associated skeleton.");
-			const FSlateIcon Icon = FSlateIconFinder::FindIconForClass(USkeletalMesh::StaticClass());
-			const FToolMenuExecuteAction UIAction = FToolMenuExecuteAction::CreateLambda([](const FToolMenuContext& InMenuContext)
+			const UContentBrowserAssetContextMenuContext* Context = InSection.FindContext<UContentBrowserAssetContextMenuContext>();
+			if (Context && Context->bCanBeModified)
 			{
-				if (const UContentBrowserAssetContextMenuContext* Context = InMenuContext.FindContext<UContentBrowserAssetContextMenuContext>())
+				const TAttribute<FText> Label = LOCTEXT("ConvertToSkeletalMesh", "Convert to Skeletal Mesh");
+				const TAttribute<FText> ToolTip = LOCTEXT("ConvertToSkeletalMeshToolTip", "Converts a non-Nanite static mesh to a skeletal mesh with an associated skeleton.");
+				const FSlateIcon Icon = FSlateIconFinder::FindIconForClass(USkeletalMesh::StaticClass());
+				const FToolMenuExecuteAction UIAction = FToolMenuExecuteAction::CreateLambda([](const FToolMenuContext& InMenuContext)
 				{
-					ConvertStaticMeshAssetsToSkeletalMeshesInteractive(
-						Context->GetSelectedAssetsOfType(UStaticMesh::StaticClass()));
-				}
-			});
-
-			InSection.AddMenuEntry("ConvertToSkeletalMesh", Label, ToolTip, Icon, UIAction);
+					if (const UContentBrowserAssetContextMenuContext* Context = InMenuContext.FindContext<UContentBrowserAssetContextMenuContext>())
+					{
+						ConvertStaticMeshAssetsToSkeletalMeshesInteractive(
+							Context->GetSelectedAssetsOfType(UStaticMesh::StaticClass()));
+					}
+				});
+
+				InSection.AddMenuEntry("ConvertToSkeletalMesh", Label, ToolTip, Icon, UIAction);
+			}
 		}));
 	}
 }
