SkeletalMesh Modeling Mode: add support for extending SkeletalMeshModelingMode with additional Tools based on IModularFeature framework

diff --git a/Engine/Plugins/Experimental/Animation/SkeletalMeshModelingTools/Source/SkeletalMeshModelingTools/Private/SkeletalMeshModelingToolsEditorMode.cpp b/Engine/Plugins/Experimental/Animation/SkeletalMeshModelingTools/Source/SkeletalMeshModelingTools/Private/SkeletalMeshModelingToolsEditorMode.cpp
--- a/Engine/Plugins/Experimental/Animation/SkeletalMeshModelingTools/Source/SkeletalMeshModelingTools/Private/SkeletalMeshModelingToolsEditorMode.cpp
+++ b/Engine/Plugins/Experimental/Animation/SkeletalMeshModelingTools/Source/SkeletalMeshModelingTools/Private/SkeletalMeshModelingToolsEditorMode.cpp
@@ -12,6 +12,7 @@
 #include "EdModeInteractiveToolsContext.h"
 #include "EditorModeManager.h"
 #include "EditorViewportClient.h"
+#include "Features/IModularFeatures.h"
 
 // Stylus support is currently disabled due to issues with the stylus plugin
 // We are leaving the code in this cpp file, defined out, so that it is easier to bring back if/when the stylus plugin is improved.
@@ -258,6 +259,9 @@ void USkeletalMeshModelingToolsEditorMode::Enter()
 	// Skeleton Editing
 	RegisterTool(ToolManagerCommands.BeginSkeletonEditingTool, TEXT("BeginSkeletonEditingTool"), NewObject<USkeletonEditingToolBuilder>());
 
+	// register extensions
+	RegisterExtensions();
+	
 	// highlights skin weights tool by default
 	GetInteractiveToolsContext()->ToolManager->SelectActiveToolType(EToolSide::Left, TEXT("BeginSkinWeightsPaintTool"));
 
@@ -298,6 +302,71 @@ void USkeletalMeshModelingToolsEditorMode::Exit()
 	UEdMode::Exit();
 }
 
+void USkeletalMeshModelingToolsEditorMode::RegisterExtensions()
+{
+	TArray<ISkeletalMeshModelingModeToolExtension*> Extensions = IModularFeatures::Get().GetModularFeatureImplementations<ISkeletalMeshModelingModeToolExtension>(
+		ISkeletalMeshModelingModeToolExtension::GetModularFeatureName());
+	if (Extensions.IsEmpty())
+	{
+		return;
+	}
+
+	UEditorInteractiveToolsContext* ToolsContext = GetInteractiveToolsContext();
+	
+	FExtensionToolQueryInfo ExtensionQueryInfo;
+	ExtensionQueryInfo.ToolsContext = ToolsContext;
+	ExtensionQueryInfo.AssetAPI = nullptr;
+#if ENABLE_STYLUS_SUPPORT
+	ExtensionQueryInfo.StylusAPI = StylusStateTracker.Get();
+#endif
+
+	for (ISkeletalMeshModelingModeToolExtension* Extension: Extensions)
+	{
+		TArray<FExtensionToolDescription> ToolSet;
+		Extension->GetExtensionTools(ExtensionQueryInfo, ToolSet);
+		for (const FExtensionToolDescription& ToolInfo : ToolSet)
+		{
+			RegisterTool(ToolInfo.ToolCommand, ToolInfo.ToolName.ToString(), ToolInfo.ToolBuilder);
+			ExtensionToolToInfo.Add(ToolInfo.ToolName.ToString(), ToolInfo);
+		}
+
+		TArray<TSubclassOf<UToolTargetFactory>> ExtensionToolTargetFactoryClasses;
+		if (Extension->GetExtensionToolTargets(ExtensionToolTargetFactoryClasses))
+		{
+			for (const TSubclassOf<UToolTargetFactory>& ExtensionTargetFactoryClass : ExtensionToolTargetFactoryClasses)
+			{
+				ToolsContext->TargetManager->AddTargetFactory(NewObject<UToolTargetFactory>(GetToolManager(), ExtensionTargetFactoryClass.Get()));
+			}
+		}
+	}
+}
+
+bool USkeletalMeshModelingToolsEditorMode::TryGetExtensionToolCommandGetter(
+	UInteractiveToolManager* InManager, const UInteractiveTool* InTool, 
+	TFunction<const UE::IInteractiveToolCommandsInterface& ()>& GetterOut) const
+{
+	if (!ensure(InManager && InTool)
+		|| InManager->GetActiveTool(EToolSide::Mouse) != InTool)
+	{
+		return false;
+	}
+
+	const FString ToolName = InManager->GetActiveToolName(EToolSide::Mouse);
+	if (ToolName.IsEmpty())
+	{
+		return false;
+	}
+
+	const FExtensionToolDescription* ToolDescription = ExtensionToolToInfo.Find(ToolName);
+	if (!ToolDescription || !ToolDescription->ToolCommandsGetter)
+	{
+		return false;
+	}
+	
+	GetterOut = ToolDescription->ToolCommandsGetter;
+	return true;
+}
+
 void USkeletalMeshModelingToolsEditorMode::CreateToolkit()
 {
 	Toolkit = MakeShareable(new FSkeletalMeshModelingToolsEditorModeToolkit);
@@ -401,6 +470,10 @@ bool USkeletalMeshModelingToolsEditorMode::ComputeBoundingBoxForViewportFocus(AA
 void USkeletalMeshModelingToolsEditorMode::OnToolStarted(UInteractiveToolManager* Manager, UInteractiveTool* Tool)
 {
 	FSkeletalMeshModelingToolsActionCommands::UpdateToolCommandBinding(Tool, Toolkit->GetToolkitCommands(), false);
+	if (TryGetExtensionToolCommandGetter(Manager, Tool, ExtensionToolCommandsGetter) && ensure(ExtensionToolCommandsGetter))
+	{
+		ExtensionToolCommandsGetter().BindCommandsForCurrentTool(Toolkit->GetToolkitCommands(), Tool);
+	}
 
 	// deactivate SkeletonSelection when a tool is activated.
 	// each tool is responsible for activating SkeletonSelection if necessary
@@ -413,7 +486,12 @@ void USkeletalMeshModelingToolsEditorMode::OnToolStarted(UInteractiveToolManager
 void USkeletalMeshModelingToolsEditorMode::OnToolEnded(UInteractiveToolManager* Manager, UInteractiveTool* Tool)
 {
 	FSkeletalMeshModelingToolsActionCommands::UpdateToolCommandBinding(Tool, Toolkit->GetToolkitCommands(), true);
-
+	if (ExtensionToolCommandsGetter)
+	{
+		ExtensionToolCommandsGetter().UnbindActiveCommands(Toolkit->GetToolkitCommands());
+		ExtensionToolCommandsGetter = nullptr;
+	}
+	
 	// reactivate SkeletonSelection when deactivating a tool
 	if (Owner)
 	{
diff --git a/Engine/Plugins/Experimental/Animation/SkeletalMeshModelingTools/Source/SkeletalMeshModelingTools/Private/SkeletalMeshModelingToolsEditorMode.h b/Engine/Plugins/Experimental/Animation/SkeletalMeshModelingTools/Source/SkeletalMeshModelingTools/Private/SkeletalMeshModelingToolsEditorMode.h
--- a/Engine/Plugins/Experimental/Animation/SkeletalMeshModelingTools/Source/SkeletalMeshModelingTools/Private/SkeletalMeshModelingToolsEditorMode.h
+++ b/Engine/Plugins/Experimental/Animation/SkeletalMeshModelingTools/Source/SkeletalMeshModelingTools/Private/SkeletalMeshModelingToolsEditorMode.h
@@ -4,6 +4,7 @@
 
 #include "CoreMinimal.h"
 #include "EdMode.h"
+#include "SkeletalMeshModelingModeToolExtensions.h"
 
 #include "SkeletalMeshModelingToolsEditorMode.generated.h"
 
@@ -17,6 +18,10 @@ class HHitProxy;
 class UDebugSkelMeshComponent;
 class ISkeletalMeshEditor;
 enum class EToolManagerToolSwitchMode;
+namespace UE
+{
+	class IInteractiveToolCommandsInterface;
+}
 
 UCLASS()
 class USkeletalMeshModelingToolsEditorMode : 
@@ -70,4 +75,14 @@ private:
 	FDelegateHandle FromToolNotifierHandle;
 	
 	TWeakPtr<ISkeletalMeshEditorBinding> Binding;
+
+	void RegisterExtensions();
+	// Support extension tools having their own hotkey classes
+	TMap<FString, FExtensionToolDescription> ExtensionToolToInfo;
+	// Note: this will only work when the given tool is active, because we get the tool identifier
+	//  out of the manager using GetActiveToolName
+	bool TryGetExtensionToolCommandGetter(UInteractiveToolManager* InManager, const UInteractiveTool* InTool, 
+		TFunction<const UE::IInteractiveToolCommandsInterface&()>& OutGetter) const;
+	// Used to unbind extension tool commands
+	TFunction<const UE::IInteractiveToolCommandsInterface& ()> ExtensionToolCommandsGetter;
 };
diff --git a/Engine/Plugins/Experimental/Animation/SkeletalMeshModelingTools/Source/SkeletalMeshModelingTools/Private/SkeletalMeshModelingToolsEditorModeToolkit.cpp b/Engine/Plugins/Experimental/Animation/SkeletalMeshModelingTools/Source/SkeletalMeshModelingTools/Private/SkeletalMeshModelingToolsEditorModeToolkit.cpp
--- a/Engine/Plugins/Experimental/Animation/SkeletalMeshModelingTools/Source/SkeletalMeshModelingTools/Private/SkeletalMeshModelingToolsEditorModeToolkit.cpp
+++ b/Engine/Plugins/Experimental/Animation/SkeletalMeshModelingTools/Source/SkeletalMeshModelingTools/Private/SkeletalMeshModelingToolsEditorModeToolkit.cpp
@@ -5,6 +5,7 @@
 #include "SkeletalMeshModelingToolsCommands.h"
 
 #include "EdModeInteractiveToolsContext.h"
+#include "Features/IModularFeatures.h"
 #include "Framework/MultiBox/MultiBoxBuilder.h"
 #include "IDetailsView.h"
 #include "InteractiveTool.h"
@@ -18,6 +19,7 @@
 #include "Widgets/Images/SImage.h"
 #include "Widgets/Text/STextBlock.h"
 #include "ModelingToolsManagerActions.h"
+#include "SkeletalMeshModelingModeToolExtensions.h"
 #include "BaseTools/MeshSurfacePointTool.h"
 #include "SPrimaryButton.h"
 #include "Toolkits/AssetEditorModeUILayer.h"
@@ -216,6 +218,8 @@ void FSkeletalMeshModelingToolsEditorModeToolkit::RegisterPalettes()
 
 	ToolkitBuilder->SetActivePaletteOnLoad(Commands.LoadSkinTools.Get());
 	ToolkitBuilder->UpdateWidget();
+
+	RegisterExtensionsPalettes();
 	
 	// if selected palette changes, make sure we are showing the palette command buttons, which may be hidden by active Tool
 	ActivePaletteChangedHandle = ToolkitBuilder->OnActivePaletteChanged.AddLambda([this]()
@@ -224,6 +228,55 @@ void FSkeletalMeshModelingToolsEditorModeToolkit::RegisterPalettes()
 	});
 }
 
+void FSkeletalMeshModelingToolsEditorModeToolkit::RegisterExtensionsPalettes() const
+{
+	const TArray<ISkeletalMeshModelingModeToolExtension*> Extensions = IModularFeatures::Get().GetModularFeatureImplementations<ISkeletalMeshModelingModeToolExtension>(
+		ISkeletalMeshModelingModeToolExtension::GetModularFeatureName());
+
+	if (Extensions.IsEmpty())
+	{
+		return;
+	}
+
+	for (ISkeletalMeshModelingModeToolExtension* Extension: Extensions)
+	{
+		const FText ExtensionName = Extension->GetExtensionName();
+		const FText SectionName = Extension->GetToolSectionName();
+
+		FModelingModeExtensionExtendedInfo ExtensionExtendedInfo;
+		const bool bHasExtendedInfo = Extension->GetExtensionExtendedInfo(ExtensionExtendedInfo);
+
+		TSharedPtr<FUICommandInfo> PaletteCommand;
+		if (bHasExtendedInfo && ExtensionExtendedInfo.ExtensionCommand.IsValid())
+		{
+			PaletteCommand = ExtensionExtendedInfo.ExtensionCommand;
+		}
+		else
+		{
+			const FText UseTooltipText = (bHasExtendedInfo && ExtensionExtendedInfo.ToolPaletteButtonTooltip.IsEmpty() == false) ?
+				ExtensionExtendedInfo.ToolPaletteButtonTooltip : SectionName;
+			PaletteCommand = FModelingToolsManagerCommands::RegisterExtensionPaletteCommand(
+				FName(ExtensionName.ToString()),
+				SectionName, UseTooltipText, FSlateIcon());
+		}
+
+		TArray<TSharedPtr<FUICommandInfo>> PaletteItems;
+		FExtensionToolQueryInfo ExtensionQueryInfo;
+		ExtensionQueryInfo.bIsInfoQueryOnly = true;
+
+		TArray<FExtensionToolDescription> ToolSet;
+		Extension->GetExtensionTools(ExtensionQueryInfo, ToolSet);
+		
+		for (const FExtensionToolDescription& ToolInfo : ToolSet)
+		{
+			PaletteItems.Add(ToolInfo.ToolCommand);
+		}
+
+		ToolkitBuilder->AddPalette(
+			MakeShareable(new FToolPalette(PaletteCommand.ToSharedRef(), PaletteItems)));
+	}
+}
+
 void FSkeletalMeshModelingToolsEditorModeToolkit::MakeToolAcceptCancelWidget()
 {
 	SAssignNew(ViewportOverlayWidget, SHorizontalBox)
diff --git a/Engine/Plugins/Experimental/Animation/SkeletalMeshModelingTools/Source/SkeletalMeshModelingTools/Private/SkeletalMeshModelingToolsEditorModeToolkit.h b/Engine/Plugins/Experimental/Animation/SkeletalMeshModelingTools/Source/SkeletalMeshModelingTools/Private/SkeletalMeshModelingToolsEditorModeToolkit.h
--- a/Engine/Plugins/Experimental/Animation/SkeletalMeshModelingTools/Source/SkeletalMeshModelingTools/Private/SkeletalMeshModelingToolsEditorModeToolkit.h
+++ b/Engine/Plugins/Experimental/Animation/SkeletalMeshModelingTools/Source/SkeletalMeshModelingTools/Private/SkeletalMeshModelingToolsEditorModeToolkit.h
@@ -40,6 +40,7 @@ private:
 	void UpdateActiveToolProperties(UInteractiveTool* Tool);
 
 	void RegisterPalettes();
+	void RegisterExtensionsPalettes() const;
 	FDelegateHandle ActivePaletteChangedHandle;
 
 	void MakeToolAcceptCancelWidget();
diff --git a/Engine/Plugins/Experimental/Animation/SkeletalMeshModelingTools/Source/SkeletalMeshModelingTools/Public/SkeletalMeshModelingModeToolExtensions.h b/Engine/Plugins/Experimental/Animation/SkeletalMeshModelingTools/Source/SkeletalMeshModelingTools/Public/SkeletalMeshModelingModeToolExtensions.h
new file mode 100644
--- /dev/null
+++ b/Engine/Plugins/Experimental/Animation/SkeletalMeshModelingTools/Source/SkeletalMeshModelingTools/Public/SkeletalMeshModelingModeToolExtensions.h
@@ -0,0 +1,25 @@
+// Copyright Epic Games, Inc. All Rights Reserved.
+
+#pragma once
+
+#include "Features/IModularFeature.h"
+#include "ModelingModeToolExtensions.h"
+#include "Internationalization/Text.h"
+
+/**
+ * ISkeletalMeshModelingModeToolExtension uses the IModularFeature API to allow a Plugin to provide
+ * a set of InteractiveTool's to be exposed in SkeletalMesh Modeling Mode. The Tools will be 
+ * included in a section of the SkeletalMesh Modeling Mode tool list, based on GetToolSectionName().
+ * 
+ */
+class ISkeletalMeshModelingModeToolExtension : public IModelingModeToolExtension
+{
+public:
+	virtual ~ISkeletalMeshModelingModeToolExtension() override {}
+
+	static FName GetModularFeatureName()
+	{
+		static FName FeatureName( TEXT("SkeletalMeshModelingModeToolExtension") );
+		return FeatureName;
+	}
+};
