Skeleton Editing Tool: expose global reset widget

#rnx

diff --git a/Engine/Plugins/Experimental/Animation/SkeletalMeshModelingTools/Source/SkeletalMeshModelingTools/Private/DetailCustomization/SkeletonEditingToolPropertyCustomizations.cpp b/Engine/Plugins/Experimental/Animation/SkeletalMeshModelingTools/Source/SkeletalMeshModelingTools/Private/DetailCustomization/SkeletonEditingToolPropertyCustomizations.cpp
--- a/Engine/Plugins/Experimental/Animation/SkeletalMeshModelingTools/Source/SkeletalMeshModelingTools/Private/DetailCustomization/SkeletonEditingToolPropertyCustomizations.cpp
+++ b/Engine/Plugins/Experimental/Animation/SkeletalMeshModelingTools/Source/SkeletalMeshModelingTools/Private/DetailCustomization/SkeletonEditingToolPropertyCustomizations.cpp
@@ -271,6 +271,10 @@ void FSkeletonEditingPropertiesDetailCustomization::CustomizeValueGet(SAdvancedT
 					}
 					break;
 				default:
+					if (!Transform.Equals(FTransform::Identity))
+					{
+						return true;
+					}
 					break;
 			}
 		}
@@ -405,6 +409,7 @@ void FSkeletonEditingPropertiesDetailCustomization::CustomizeValueSet(SAdvancedT
 			OutTransform.SetScale3D(FVector::OneVector);
 			break;
 		default:
+			OutTransform = FTransform::Identity;
 			break;
 		}
 		
