Skinweight Editing: Fixed the spinbox for clamp and clamp select to use int32 to match the underlying type.
- Changed the default for clamp select to be 9 based on user feedback.

#rnx

diff --git a/Engine/Plugins/Experimental/Animation/SkeletalMeshModelingTools/Source/SkeletalMeshModelingTools/Private/DetailCustomization/SkinWeightDetailCustomization.cpp b/Engine/Plugins/Experimental/Animation/SkeletalMeshModelingTools/Source/SkeletalMeshModelingTools/Private/DetailCustomization/SkinWeightDetailCustomization.cpp
--- a/Engine/Plugins/Experimental/Animation/SkeletalMeshModelingTools/Source/SkeletalMeshModelingTools/Private/DetailCustomization/SkinWeightDetailCustomization.cpp
+++ b/Engine/Plugins/Experimental/Animation/SkeletalMeshModelingTools/Source/SkeletalMeshModelingTools/Private/DetailCustomization/SkinWeightDetailCustomization.cpp
@@ -586,18 +586,18 @@ void FSkinWeightDetailCustomization::AddSelectionUI(IDetailLayoutBuilder& Detail
 			+SHorizontalBox::Slot()
 			.Padding(2.f, WeightEditVerticalPadding)
 			[
-				SNew(SSpinBox<float>)
+				SNew(SSpinBox<int32>)
 				.MinValue(1)
 				.MaxValue(24)
 				.Value_Lambda([this]()
 				{
 					return ToolSettings->ClampSelectValue;
 				})
-				.OnValueChanged_Lambda([this](float NewValue)
+				.OnValueChanged_Lambda([this](int32 NewValue)
 				{
 					ToolSettings->ClampSelectValue = NewValue;
 				})
-				.OnValueCommitted_Lambda([this](float NewValue, ETextCommit::Type CommitType)
+				.OnValueCommitted_Lambda([this](int32 NewValue, ETextCommit::Type CommitType)
 				{
 					ToolSettings->SaveConfig();
 				})
@@ -1198,18 +1198,18 @@ void FSkinWeightDetailCustomization::AddSelectionUI(IDetailLayoutBuilder& Detail
 		+SHorizontalBox::Slot()
 		.Padding(2.f, WeightEditVerticalPadding)
 		[
-			SNew(SSpinBox<float>)
+			SNew(SSpinBox<int32>)
 			.MinValue(1)
 			.MaxValue(24)
 			.Value_Lambda([this]()
 			{
 				return ToolSettings->ClampValue;
 			})
-			.OnValueChanged_Lambda([this](float NewValue)
+			.OnValueChanged_Lambda([this](int32 NewValue)
 			{
 				ToolSettings->ClampValue = NewValue;
 			})
-			.OnValueCommitted_Lambda([this](float NewValue, ETextCommit::Type CommitType)
+			.OnValueCommitted_Lambda([this](int32 NewValue, ETextCommit::Type CommitType)
 			{
 				ToolSettings->SaveConfig();
 			})
