Skeleton Editing Tool: multi-bones drag & drop

diff --git a/Engine/Plugins/Experimental/Animation/SkeletalMeshModelingTools/Source/SkeletalMeshModelingTools/Private/SReferenceSkeletonTree.cpp b/Engine/Plugins/Experimental/Animation/SkeletalMeshModelingTools/Source/SkeletalMeshModelingTools/Private/SReferenceSkeletonTree.cpp
--- a/Engine/Plugins/Experimental/Animation/SkeletalMeshModelingTools/Source/SkeletalMeshModelingTools/Private/SReferenceSkeletonTree.cpp
+++ b/Engine/Plugins/Experimental/Animation/SkeletalMeshModelingTools/Source/SkeletalMeshModelingTools/Private/SReferenceSkeletonTree.cpp
@@ -175,11 +175,20 @@ TSharedRef<SWidget> SReferenceSkeletonRow::GenerateWidgetForColumn(const FName&
 	];
 }
 
-TSharedRef<FBoneItemDragDropOp> FBoneItemDragDropOp::New(const TWeakPtr<FBoneElement>& InElement)
+TSharedRef<FBoneItemDragDropOp> FBoneItemDragDropOp::New(TArrayView<const TSharedPtr<FBoneElement>> InElements)
 {
 	TSharedRef<FBoneItemDragDropOp> Operation = MakeShared<FBoneItemDragDropOp>();
-		Operation->Element = InElement;
-		Operation->Construct();
+
+	for (const TSharedPtr<FBoneElement>& BoneElement: InElements)
+	{
+		const bool bContainsParent = InElements.Contains(BoneElement->Parent);
+		if (!bContainsParent)
+		{
+			Operation->Elements.Add(BoneElement);
+		}
+	}
+	
+	Operation->Construct();
 	return Operation;
 }
 
@@ -190,7 +199,10 @@ TSharedPtr<SWidget> FBoneItemDragDropOp::GetDefaultDecorator() const
         .BorderImage(FAppStyle::GetBrush("Menu.Background"))
         [
             SNew(STextBlock)
-            .Text(FText::FromName(Element.Pin()->BoneName))
+            .Text_Lambda([this]()
+            {
+	            return !Elements.IsEmpty() && Elements[0].IsValid() ? FText::FromName(Elements[0].Pin()->BoneName) : FText::GetEmpty();
+            })
         ];
 }
 
@@ -978,14 +990,14 @@ FReply SReferenceSkeletonTree::OnKeyDown(const FGeometry& MyGeometry, const FKey
 FReply SReferenceSkeletonTree::OnDragDetected(const FGeometry& MyGeometry, const FPointerEvent& MouseEvent)
 {
 	const TArray<TSharedPtr<FBoneElement>> BoneElements = GetSelectedItems();
-	if (BoneElements.Num() != 1)
+	if (BoneElements.IsEmpty())
 	{
 		return FReply::Unhandled();
 	}
 	
 	if (MouseEvent.IsMouseButtonDown(EKeys::LeftMouseButton))
 	{
-		const TSharedRef<FBoneItemDragDropOp> DragDropOp = FBoneItemDragDropOp::New(BoneElements[0]);
+		const TSharedRef<FBoneItemDragDropOp> DragDropOp = FBoneItemDragDropOp::New(BoneElements);
 		return FReply::Handled().BeginDragDrop(DragDropOp);
 	}
 	
@@ -1018,20 +1030,28 @@ FReply SReferenceSkeletonTree::OnAcceptDrop(const FDragDropEvent& DragDropEvent,
 		return FReply::Handled();
 	}
 	
-	const FName BoneName = DragDropOp->Element.Pin()->BoneName;
+	TArray<FName> BoneNames;
+	BoneNames.Reserve(DragDropOp->Elements.Num());
+	for (const TWeakPtr<FBoneElement>& BoneElement: DragDropOp->Elements)
+	{
+		if (BoneElement.IsValid())
+		{
+			BoneNames.Add(BoneElement.Pin()->BoneName);
+		}
+	}
 	const FName ParentName = TargetItem->BoneName;
 
 	BeginChange();
 	
-	const bool bParented = Modifier->ParentBone(BoneName, ParentName);
+	const bool bParented = Modifier->ParentBones(BoneNames, {ParentName});
 	if (bParented)
 	{
 		RefreshTreeView(true);
-		SelectItemFromNames({BoneName});
+		SelectItemFromNames(BoneNames);
 		if (Notifier.IsValid())
 		{
-			Notifier->Notify( {BoneName}, ESkeletalMeshNotifyType::HierarchyChanged);
-			Notifier->Notify( {BoneName}, ESkeletalMeshNotifyType::BonesSelected);
+			Notifier->Notify( BoneNames, ESkeletalMeshNotifyType::HierarchyChanged);
+			Notifier->Notify( BoneNames, ESkeletalMeshNotifyType::BonesSelected);
 		}
 		EndChange();
 	}
diff --git a/Engine/Plugins/Experimental/Animation/SkeletalMeshModelingTools/Source/SkeletalMeshModelingTools/Private/SReferenceSkeletonTree.h b/Engine/Plugins/Experimental/Animation/SkeletalMeshModelingTools/Source/SkeletalMeshModelingTools/Private/SReferenceSkeletonTree.h
--- a/Engine/Plugins/Experimental/Animation/SkeletalMeshModelingTools/Source/SkeletalMeshModelingTools/Private/SReferenceSkeletonTree.h
+++ b/Engine/Plugins/Experimental/Animation/SkeletalMeshModelingTools/Source/SkeletalMeshModelingTools/Private/SReferenceSkeletonTree.h
@@ -110,13 +110,13 @@ class FBoneItemDragDropOp : public FDecoratedDragDropOp
 public:
 	DRAG_DROP_OPERATOR_TYPE(FBoneItemDragDropOp, FDecoratedDragDropOp)
 
-	static TSharedRef<FBoneItemDragDropOp> New(const TWeakPtr<FBoneElement>& InElement);
+	static TSharedRef<FBoneItemDragDropOp> New(TArrayView<const TSharedPtr<FBoneElement>> InElements);
 
 	// FDecoratedDragDropOp interface
 	virtual TSharedPtr<SWidget> GetDefaultDecorator() const override;
 	// End FDecoratedDragDropOp interface
 	
-	TWeakPtr<FBoneElement> Element;
+	TArray< TWeakPtr<FBoneElement> > Elements;
 };
 
 class SReferenceSkeletonRow : public SMultiColumnTableRow<TSharedPtr<FBoneElement>>
