Skeletal Mesh Editor Tools: moved from Experimental

#virtualized

diff --git a/Engine/Plugins/Animation/SkeletalMeshModelingTools/SkeletalMeshModelingTools.uplugin b/Engine/Plugins/Animation/SkeletalMeshModelingTools/SkeletalMeshModelingTools.uplugin
new file mode 100644
--- /dev/null
+++ b/Engine/Plugins/Animation/SkeletalMeshModelingTools/SkeletalMeshModelingTools.uplugin
@@ -0,0 +1,31 @@
+﻿{
+	"FileVersion": 3,
+	"Version": 1,
+	"VersionName": "1.0",
+	"FriendlyName": "Skeletal Mesh Editing Tools",
+	"Description": "Create skeletons, paint skin weights and edit skeletal meshes.",
+	"Category" : "Animation",
+	"CreatedBy" : "Epic Games, Inc.",
+	"CreatedByURL": "https://epicgames.com/",
+	"DocsURL": "",
+	"MarketplaceURL": "",
+	"EnabledByDefault" : true,
+	"CanContainContent": true,
+	"IsExperimentalVersion": false,
+	"IsBetaVersion": false,
+	"Installed": false,
+	"Modules": [
+		{
+			"Name": "SkeletalMeshModelingTools",
+			"Type": "Editor",
+			"LoadingPhase": "Default"
+		}
+	],
+	"Plugins": [
+		{
+			"Name": "ModelingToolsEditorMode",
+			"Enabled": true,
+			"TargetAllowList": [ "Editor" ]
+		}
+	]
+}
diff --git a/Engine/Plugins/Animation/SkeletalMeshModelingTools/Source/SkeletalMeshModelingTools/Private/DetailCustomization/SkeletonEditingToolPropertyCustomizations.cpp b/Engine/Plugins/Animation/SkeletalMeshModelingTools/Source/SkeletalMeshModelingTools/Private/DetailCustomization/SkeletonEditingToolPropertyCustomizations.cpp
new file mode 100644
--- /dev/null
+++ b/Engine/Plugins/Animation/SkeletalMeshModelingTools/Source/SkeletalMeshModelingTools/Private/DetailCustomization/SkeletonEditingToolPropertyCustomizations.cpp
@@ -0,0 +1,1104 @@
+// Copyright Epic Games, Inc. All Rights Reserved.
+
+#include "SkeletonEditingToolPropertyCustomizations.h"
+
+#include "DetailLayoutBuilder.h"
+#include "Templates/UniquePtr.h"
+#include "SkeletalMesh/SkeletonEditingTool.h"
+#include "Widgets/SBoxPanel.h"
+#include "Widgets/Input/SEditableTextBox.h"
+#include "Widgets/Input/SSegmentedControl.h"
+#include "HAL/PlatformApplicationMisc.h"
+#include "InteractiveToolManager.h"
+#include "ModelingToolsEditorModeStyle.h"
+#include "Selection/PolygonSelectionMechanic.h"
+
+
+#define LOCTEXT_NAMESPACE "SkeletonEditingToolCustomization"
+
+TSharedRef<IDetailCustomization> FSkeletonEditingPropertiesDetailCustomization::MakeInstance()
+{
+	return MakeShareable(new FSkeletonEditingPropertiesDetailCustomization);
+}
+
+TAttribute<EVisibility> ISkeletonEditingPropertiesDetailCustomization::GetCreationVisibility() const
+{
+	return TAttribute<EVisibility>::CreateLambda([this]()
+	{
+		if (!Tool.IsValid())
+		{
+			return EVisibility::Hidden;
+		}
+		return Tool->GetOperation() == EEditingOperation::Create ? EVisibility::Visible : EVisibility::Hidden;
+	});
+}
+
+TAttribute<EVisibility> ISkeletonEditingPropertiesDetailCustomization::GetEditionVisibility() const
+{
+	return TAttribute<EVisibility>::CreateLambda([this]()
+	{
+		if (!Tool.IsValid())
+		{
+			return EVisibility::Hidden;
+		}
+		return Tool->GetOperation() == EEditingOperation::Create ? EVisibility::Hidden : EVisibility::Visible;
+	});
+}
+
+TAttribute<bool> ISkeletonEditingPropertiesDetailCustomization::IsEnabledBySelection() const
+{
+	return TAttribute<bool>::CreateLambda([this]()
+	{
+		if (!Tool.IsValid())
+		{
+			return false;
+		}
+		return !Tool->GetSelection().IsEmpty() && Tool->GetOperation() != EEditingOperation::Parent;
+	});
+}
+
+void ISkeletonEditingPropertiesDetailCustomization::UpdateProperties(
+	IDetailLayoutBuilder& DetailBuilder,
+	const TAttribute<EVisibility>& InVisibility,
+	const TAttribute<bool>& InEnabled) const
+{
+	for (const FName& PropName: GetProperties())
+	{
+		const TSharedRef<IPropertyHandle> Handle = DetailBuilder.GetProperty(PropName);
+		if (IDetailPropertyRow* Row = DetailBuilder.EditDefaultProperty(Handle))
+		{
+			Row->Visibility(InVisibility);
+			if (InEnabled.IsSet())
+			{
+				Row->IsEnabled(InEnabled);
+			}
+		}
+	}
+}
+
+void FSkeletonEditingPropertiesDetailCustomization::CustomizeDetails(IDetailLayoutBuilder& DetailBuilder)
+{
+	const TWeakObjectPtr<USkeletonEditingProperties> Properties = GetCustomizedObject<USkeletonEditingProperties>(DetailBuilder);
+	if (!GetParentTool(Properties).IsValid())
+	{
+		return;
+	}
+	
+	RelativeArray.Init(true, 4);
+	
+	IDetailCategoryBuilder& CategoryBuilder = DetailBuilder.EditCategory(TEXT("Details"));
+	const TSharedRef<IPropertyHandle> NodePropHandle = DetailBuilder.GetProperty(GET_MEMBER_NAME_CHECKED(USkeletonEditingProperties, Name));
+	if (NodePropHandle->IsValidHandle())
+	{
+		TSharedRef<SWidget> ValueWidget = NodePropHandle->CreatePropertyValueWidget();
+		ValueWidget->SetEnabled(TAttribute<bool>::Create([this]{ return !Tool->GetSelection().IsEmpty(); }));
+		static const FText MultiValues = LOCTEXT("MultipleValues", "Multiple Values");
+		
+		DetailBuilder.AddPropertyToCategory(NodePropHandle)
+			.CustomWidget()
+			.NameContent()
+			[
+				NodePropHandle->CreatePropertyNameWidget()
+			]
+			.ValueContent()
+			[
+				SNew(SEditableTextBox)
+					.Font(FAppStyle::GetFontStyle( TEXT( "PropertyWindow.NormalFont" ) ))
+					.SelectAllTextWhenFocused(true)
+					.ClearKeyboardFocusOnCommit(false)
+					.SelectAllTextOnCommit(true)
+					.Text_Lambda([this]()
+					{
+						const TArray<FName>& Bones = Tool->GetSelection();
+						if (Bones.Num() == 0)
+						{
+							return FText::FromName(NAME_None);
+						}
+						if (Bones.Num() == 1)
+						{
+							return FText::FromName(Bones[0]);
+						}
+						return MultiValues;
+					})
+					.OnTextCommitted_Lambda([this, NodePropHandle](const FText& NewText, ETextCommit::Type InCommitType) 
+					{
+						if (InCommitType == ETextCommit::OnCleared || NewText.EqualTo(MultiValues))
+						{
+							return;
+						}
+						FText CurrentText; NodePropHandle->GetValueAsFormattedText(CurrentText);
+						if (!NewText.ToString().Equals(CurrentText.ToString(), ESearchCase::CaseSensitive))
+						{
+							NodePropHandle->SetValueFromFormattedString(NewText.ToString());
+						}
+					})
+					.IsEnabled_Lambda([this] { return !Tool->GetSelection().IsEmpty(); })
+			];
+	}
+	
+	SAdvancedTransformInputBox<FTransform>::FArguments TransformWidgetArgs = SAdvancedTransformInputBox<FTransform>::FArguments()
+		.IsEnabled(true)
+		.DisplayRelativeWorld(true)
+		.Font(IDetailLayoutBuilder::GetDetailFont())
+		.DisplayScaleLock(false)
+		.AllowEditRotationRepresentation(false);
+	// relative/world
+	TransformWidgetArgs
+	.OnGetIsComponentRelative_Lambda( [this](ESlateTransformComponent::Type InComponent)
+	{
+	   return RelativeArray[InComponent];
+	})
+	.OnIsComponentRelativeChanged_Lambda( [this](ESlateTransformComponent::Type InComponent, bool bIsRelative)
+	{
+		RelativeArray[InComponent] = bIsRelative;
+	});
+	
+	// get bones transforms
+	CustomizeValueGet(TransformWidgetArgs);
+		
+	// set bones transforms
+	CustomizeValueSet(TransformWidgetArgs);
+	// copy/paste values
+	CustomizeClipboard(TransformWidgetArgs);
+	// enabled
+	TransformWidgetArgs
+	.IsEnabled_Lambda([this]()
+	{
+		return !Tool->GetSelection().IsEmpty();
+	});
+	
+	SAdvancedTransformInputBox<FTransform>::ConstructGroupedTransformRows(
+		CategoryBuilder, 
+		LOCTEXT("ReferenceTransform", "Transform"),
+		LOCTEXT("ReferenceBoneTransformTooltip", "The reference transform of the bone"), 
+		TransformWidgetArgs);
+	
+	const TSharedRef<IPropertyHandle> Handle = DetailBuilder.GetProperty(GET_MEMBER_NAME_CHECKED(USkeletonEditingProperties, bUpdateChildren));
+	if (IDetailPropertyRow* Row = DetailBuilder.EditDefaultProperty(Handle))
+	{
+		Row->IsEnabled(IsEnabledBySelection());
+	}
+}
+
+const TArray<FName>& FSkeletonEditingPropertiesDetailCustomization::GetProperties() const
+{
+	static const TArray< FName > Properties(
+	{
+		GET_MEMBER_NAME_CHECKED(USkeletonEditingProperties, Name),
+		GET_MEMBER_NAME_CHECKED(USkeletonEditingProperties, bUpdateChildren),
+		GET_MEMBER_NAME_CHECKED(USkeletonEditingProperties, AxisLength),
+		GET_MEMBER_NAME_CHECKED(USkeletonEditingProperties, AxisThickness),
+	});
+	return Properties;
+}
+
+void FSkeletonEditingPropertiesDetailCustomization::CustomizeValueGet(SAdvancedTransformInputBox<FTransform>::FArguments& InOutArgs)
+{
+	if (!Tool.IsValid())
+	{
+		return;
+	}
+	
+	auto GetNumericValue = [this](
+	   const FName InBoneName,
+	   ESlateTransformComponent::Type InComponent,
+	   ESlateRotationRepresentation::Type InRepresentation,
+	   ESlateTransformSubComponent::Type InSubComponent)
+	{
+		const bool bWorld = !RelativeArray[InComponent];
+		return SAdvancedTransformInputBox<FTransform>::GetNumericValueFromTransform(
+			  Tool->GetTransform(InBoneName, bWorld),
+			  InComponent,
+			  InRepresentation,
+			  InSubComponent);
+	};
+
+	InOutArgs.OnGetNumericValue_Lambda( [this, GetNumericValue](
+		ESlateTransformComponent::Type InComponent,
+		ESlateRotationRepresentation::Type InRepresentation,
+		ESlateTransformSubComponent::Type InSubComponent)
+	{
+		const TArray<FName>& Bones = Tool->GetSelection();
+		if (Bones.IsEmpty())
+		{
+			return SAdvancedTransformInputBox<FTransform>::GetNumericValueFromTransform(
+				FTransform::Identity, InComponent, InRepresentation, InSubComponent);
+		}
+		
+		TOptional<FVector::FReal> Value = GetNumericValue(Bones[0], InComponent, InRepresentation, InSubComponent);
+		if (Value)
+		{
+			for (int32 Index = 1; Index < Bones.Num(); Index++)
+			{
+				const TOptional<FVector::FReal> NextValue = GetNumericValue(Bones[Index], InComponent, InRepresentation, InSubComponent);
+				if (NextValue)
+				{
+					if (!FMath::IsNearlyEqual(*Value, *NextValue))
+					{
+						return TOptional<FVector::FReal>();
+					}
+				}
+			}
+		}
+		return Value;
+	} );
+	
+	InOutArgs.DiffersFromDefault_Lambda([this](ESlateTransformComponent::Type InComponent)
+	{
+		const TArray<FName>& Bones = Tool->GetSelection();
+		for (const FName& BoneName: Bones)
+		{
+			const bool bWorld = !RelativeArray[InComponent];
+			const FTransform& Transform = Tool->GetTransform(BoneName, bWorld);
+			switch (InComponent)
+			{
+				case ESlateTransformComponent::Location:
+					if (!Transform.GetLocation().Equals(FVector::ZeroVector))
+					{
+						return true;
+					}
+					break;
+				case ESlateTransformComponent::Rotation:
+					if(!Transform.GetRotation().Equals(FQuat::Identity))
+					{
+						return true;
+					}
+					break;
+				case ESlateTransformComponent::Scale:
+					if(!Transform.GetScale3D().Equals(FVector::OneVector))
+					{
+						return true;
+					}
+					break;
+				default:
+					if (!Transform.Equals(FTransform::Identity))
+					{
+						return true;
+					}
+					break;
+			}
+		}
+		return false;
+	});
+}
+
+void FSkeletonEditingPropertiesDetailCustomization::CustomizeValueSet(SAdvancedTransformInputBox<FTransform>::FArguments& InOutArgs)
+{
+	if (!Tool.IsValid())
+	{
+		return;
+	}
+	
+	auto PrepareNumericValueChanged = [this]( const FName InBoneName,
+											   ESlateTransformComponent::Type InComponent,
+											   ESlateRotationRepresentation::Type InRepresentation,
+											   ESlateTransformSubComponent::Type InSubComponent,
+											   FTransform::FReal InValue)
+	{
+		const bool bWorld = !RelativeArray[InComponent];
+		const FTransform& InTransform = Tool->GetTransform(InBoneName, bWorld);
+		FTransform OutTransform = InTransform;
+		SAdvancedTransformInputBox<FTransform>::ApplyNumericValueChange(OutTransform, InValue, InComponent, InRepresentation, InSubComponent);
+		return MakeTuple(InTransform, OutTransform);
+	};
+
+	InOutArgs.OnNumericValueChanged_Lambda([this, PrepareNumericValueChanged](
+		ESlateTransformComponent::Type InComponent,
+		ESlateRotationRepresentation::Type InRepresentation,
+		ESlateTransformSubComponent::Type InSubComponent,
+		FVector::FReal InValue)
+	{
+		const TArray<FName>& Bones = Tool->GetSelection();
+		TArray<FName> BonesToMove; BonesToMove.Reserve(Bones.Num());
+		TArray<FTransform> UpdatedTransforms; UpdatedTransforms.Reserve(Bones.Num());
+		
+		FTransform CurrentTransform, UpdatedTransform;
+		for (const FName& BoneName: Bones)
+		{
+			Tie(CurrentTransform, UpdatedTransform) =
+				PrepareNumericValueChanged(BoneName, InComponent, InRepresentation, InSubComponent, InValue);
+			if (!UpdatedTransform.Equals(CurrentTransform))
+			{
+				BonesToMove.Add(BoneName);
+				UpdatedTransforms.Add(UpdatedTransform);
+			}
+		}
+
+		if (!BonesToMove.IsEmpty())
+		{
+			if (!ActiveChange.IsValid())
+			{
+				ActiveChange = MakeUnique<SkeletonEditingTool::FRefSkeletonChange>(Tool.Get());
+			}
+		
+			const bool bWorld = !RelativeArray[InComponent];
+			Tool->SetTransforms(BonesToMove, UpdatedTransforms, bWorld);
+		}
+	});
+
+	InOutArgs.OnNumericValueCommitted_Lambda([this, PrepareNumericValueChanged](
+		ESlateTransformComponent::Type InComponent,
+		ESlateRotationRepresentation::Type InRepresentation,
+		ESlateTransformSubComponent::Type InSubComponent,
+		FVector::FReal InValue,
+		ETextCommit::Type InCommitType)
+	{
+		const TArray<FName>& Bones = Tool->GetSelection();
+
+		TArray<FName> BonesToMove;
+		BonesToMove.Reserve(Bones.Num());
+		
+		TArray<FTransform> UpdatedTransforms;
+		UpdatedTransforms.Reserve(Bones.Num());
+		
+		FTransform CurrentTransform, UpdatedTransform;
+		for (const FName& BoneName: Bones)
+		{
+			Tie(CurrentTransform, UpdatedTransform) =
+				PrepareNumericValueChanged(BoneName, InComponent, InRepresentation, InSubComponent, InValue);
+			if (!UpdatedTransform.Equals(CurrentTransform))
+			{
+				BonesToMove.Add(BoneName);
+				UpdatedTransforms.Add(UpdatedTransform);
+			}
+		}
+
+		if (!BonesToMove.IsEmpty())
+		{
+			if (!ActiveChange.IsValid())
+			{
+				ActiveChange = MakeUnique<SkeletonEditingTool::FRefSkeletonChange>(Tool.Get());
+			}
+
+			const bool bWorld = !RelativeArray[InComponent];
+			Tool->SetTransforms(BonesToMove, UpdatedTransforms, bWorld);
+		}
+
+		if (ActiveChange.IsValid())
+		{
+			// send transaction
+			if (UInteractiveToolManager* ToolManager = Tool->GetToolManager())
+			{
+				ActiveChange->StoreSkeleton(Tool.Get());
+
+				static const FText TransactionDesc = LOCTEXT("ChangeNumericValue", "Change Numeric Value");
+				ToolManager->BeginUndoTransaction(TransactionDesc);
+				ToolManager->EmitObjectChange(Tool.Get(), MoveTemp(ActiveChange), TransactionDesc);
+				ToolManager->EndUndoTransaction();
+			}
+		
+			ActiveChange.Reset();
+		}
+	});
+
+	auto PrepareNumericValueReset = [this]( const FName InBoneName, ESlateTransformComponent::Type InComponent)
+	{
+		const bool bWorld = !RelativeArray[InComponent];
+		const FTransform& InTransform = Tool->GetTransform(InBoneName, bWorld);
+		FTransform OutTransform = InTransform;
+
+		switch (InComponent)
+		{
+		case ESlateTransformComponent::Location:
+			OutTransform.SetLocation(FVector::ZeroVector);
+			break;
+		case ESlateTransformComponent::Rotation:
+			OutTransform.SetRotation(FQuat::Identity);
+			break;
+		case ESlateTransformComponent::Scale:
+			OutTransform.SetScale3D(FVector::OneVector);
+			break;
+		default:
+			OutTransform = FTransform::Identity;
+			break;
+		}
+		
+		return MakeTuple(InTransform, OutTransform);
+	};
+	
+	InOutArgs.OnResetToDefault_Lambda([this, PrepareNumericValueReset](ESlateTransformComponent::Type InComponent)
+	{
+		const TArray<FName>& Bones = Tool->GetSelection();
+
+		TArray<FName> BonesToMove;
+		BonesToMove.Reserve(Bones.Num());
+		
+		TArray<FTransform> UpdatedTransforms;
+		UpdatedTransforms.Reserve(Bones.Num());
+		
+		FTransform CurrentTransform, UpdatedTransform;
+		for (const FName& BoneName: Bones)
+		{
+			Tie(CurrentTransform, UpdatedTransform) = PrepareNumericValueReset(BoneName, InComponent);
+			if (!UpdatedTransform.Equals(CurrentTransform))
+			{
+				BonesToMove.Add(BoneName);
+				UpdatedTransforms.Add(UpdatedTransform);
+			}
+		}
+
+		if (!BonesToMove.IsEmpty())
+		{
+			if (!ActiveChange.IsValid())
+			{
+				ActiveChange = MakeUnique<SkeletonEditingTool::FRefSkeletonChange>(Tool.Get());
+			}
+
+			const bool bWorld = !RelativeArray[InComponent];
+			Tool->SetTransforms(BonesToMove, UpdatedTransforms, bWorld);
+		}
+
+		if (ActiveChange.IsValid())
+		{
+			// send transaction
+			if (UInteractiveToolManager* ToolManager = Tool->GetToolManager())
+			{
+				ActiveChange->StoreSkeleton(Tool.Get());
+
+				static const FText TransactionDesc = LOCTEXT("ResetNumericValue", "Reset Numeric Value");
+				ToolManager->BeginUndoTransaction(TransactionDesc);
+				ToolManager->EmitObjectChange(Tool.Get(), MoveTemp(ActiveChange), TransactionDesc);
+				ToolManager->EndUndoTransaction();
+			}
+				
+			ActiveChange.Reset();
+		}
+	});
+}
+
+namespace FSkeletonEditingToolClipboardLocals
+{
+	
+template<typename DataType>
+void GetContentFromData(const DataType& InData, FString& Content)
+{
+	TBaseStructure<DataType>::Get()->ExportText(Content, &InData, &InData, nullptr, PPF_None, nullptr);
+}
+
+class FSkeletonEditingToolBoneErrorPipe : public FOutputDevice
+{
+public:
+	int32 NumErrors;
+
+	FSkeletonEditingToolBoneErrorPipe()
+		: FOutputDevice()
+		, NumErrors(0)
+	{}
+
+	virtual void Serialize(const TCHAR* V, ELogVerbosity::Type Verbosity, const class FName& Category) override
+	{
+		NumErrors++;
+	}
+};
+
+template<typename DataType>
+bool GetDataFromContent(const FString& Content, DataType& OutData)
+{
+	FSkeletonEditingToolBoneErrorPipe ErrorPipe;
+	static UScriptStruct* DataStruct = TBaseStructure<DataType>::Get();
+	DataStruct->ImportText(*Content, &OutData, nullptr, PPF_None, &ErrorPipe, DataStruct->GetName(), true);
+	return (ErrorPipe.NumErrors == 0);
+}
+	
+}
+
+void FSkeletonEditingPropertiesDetailCustomization::CustomizeClipboard(SAdvancedTransformInputBox<FTransform>::FArguments& InOutArgs)
+{
+	if (!Tool.IsValid())
+	{
+		return;
+	}
+	
+	auto OnCopyToClipboard = [this](const FName InBoneName, ESlateTransformComponent::Type InComponent)
+	{
+		using namespace FSkeletonEditingToolClipboardLocals;
+
+		const bool bWorld = !RelativeArray[InComponent];
+		const FTransform& Xfo = Tool->GetTransform(InBoneName, bWorld);
+
+		FString Content;
+		switch(InComponent)
+		{
+		case ESlateTransformComponent::Location:
+			{
+				GetContentFromData(Xfo.GetLocation(), Content);
+				break;
+			}
+		case ESlateTransformComponent::Rotation:
+			{
+				GetContentFromData(Xfo.Rotator(), Content);
+				break;
+			}
+		case ESlateTransformComponent::Scale:
+			{
+				GetContentFromData(Xfo.GetScale3D(), Content);
+				break;
+			}
+		case ESlateTransformComponent::Max:
+		default:
+			{
+				GetContentFromData(Xfo, Content);
+				break;
+			}
+		}
+
+		if (!Content.IsEmpty())
+		{
+			FPlatformApplicationMisc::ClipboardCopy(*Content);
+		}
+	};
+
+	auto PreparePasteFromClipboard = [this](const FName InBoneName, ESlateTransformComponent::Type InComponent)
+	{
+		using namespace FSkeletonEditingToolClipboardLocals;
+	
+		FString Content;
+		FPlatformApplicationMisc::ClipboardPaste(Content);
+
+		const bool bWorld = !RelativeArray[InComponent];
+		FTransform Xfo = Tool->GetTransform(InBoneName, bWorld);
+		
+		if (Content.IsEmpty())
+		{
+			return MakeTuple(Xfo, Xfo);
+		}
+
+		switch(InComponent)
+		{
+		case ESlateTransformComponent::Location:
+			{
+				FVector Data = Xfo.GetLocation();
+				if (GetDataFromContent(Content, Data))
+				{
+					Xfo.SetLocation(Data);
+				}
+				break;
+			}
+		case ESlateTransformComponent::Rotation:
+			{
+				FRotator Data = Xfo.Rotator();
+				if (GetDataFromContent(Content, Data))
+				{
+					Xfo.SetRotation(FQuat(Data));
+				}
+				break;
+			}
+		case ESlateTransformComponent::Scale:
+			{
+				FVector Data = Xfo.GetScale3D();
+				if (GetDataFromContent(Content, Data))
+				{
+					Xfo.SetScale3D(Data);
+				}
+				break;
+			}
+		case ESlateTransformComponent::Max:
+		default:
+			{
+				FTransform Data = Xfo;
+				if (GetDataFromContent(Content, Data))
+				{
+					Xfo = Data;
+				}
+				break;
+			}
+		}
+
+		return MakeTuple(Tool->GetTransform(InBoneName, bWorld), Xfo);
+	};
+
+	InOutArgs.OnCopyToClipboard_Lambda( [this, OnCopyToClipboard](ESlateTransformComponent::Type InComponent)
+	{
+		const TArray<FName>& Bones = Tool->GetSelection();
+		if (!Bones.IsEmpty())
+		{
+			OnCopyToClipboard(Bones[0], InComponent);
+		}
+	})
+	.OnPasteFromClipboard_Lambda([this, PreparePasteFromClipboard](ESlateTransformComponent::Type InComponent)
+	{
+		const TArray<FName>& Bones = Tool->GetSelection();
+		TArray<FName> BonesToMove; BonesToMove.Reserve(Bones.Num());
+		TArray<FTransform> UpdatedTransforms; UpdatedTransforms.Reserve(Bones.Num());
+		
+		FTransform CurrentTransform, UpdatedTransform;
+		for (const FName& BoneName: Bones)
+		{
+			Tie(CurrentTransform, UpdatedTransform) = PreparePasteFromClipboard(BoneName, InComponent);
+			if (!UpdatedTransform.Equals(CurrentTransform))
+			{
+				BonesToMove.Add(BoneName);
+				UpdatedTransforms.Add(UpdatedTransform);
+			}
+		}
+
+		if (BonesToMove.IsEmpty())
+		{
+			return;
+		}
+		
+		if (!ActiveChange.IsValid())
+		{
+			ActiveChange = MakeUnique<SkeletonEditingTool::FRefSkeletonChange>(Tool.Get());
+		}
+		
+		const bool bWorld = !RelativeArray[InComponent];
+		Tool->SetTransforms(BonesToMove, UpdatedTransforms, bWorld);
+
+		// send transaction
+		if (UInteractiveToolManager* ToolManager = Tool->GetToolManager())
+		{
+			ActiveChange->StoreSkeleton(Tool.Get());
+			
+			static const FText TransactionDesc = LOCTEXT("PasteTransform", "Paste Transform");
+			ToolManager->BeginUndoTransaction(TransactionDesc);
+			ToolManager->EmitObjectChange(Tool.Get(), MoveTemp(ActiveChange), TransactionDesc);
+			ToolManager->EndUndoTransaction();
+		}
+		ActiveChange.Reset();
+	});
+}
+
+TSharedRef<IDetailCustomization> FMirroringPropertiesDetailCustomization::MakeInstance()
+{
+	return MakeShareable(new FMirroringPropertiesDetailCustomization);
+}
+
+const TArray<FName>& FMirroringPropertiesDetailCustomization::GetProperties() const
+{
+	static const TArray< FName > Properties(
+	{
+		GET_MEMBER_NAME_CHECKED(UMirroringProperties, Options)
+	});
+	return Properties;
+}
+
+void FMirroringPropertiesDetailCustomization::CustomizeDetails(IDetailLayoutBuilder& DetailBuilder)
+{
+	TWeakObjectPtr<UMirroringProperties> Properties = GetCustomizedObject<UMirroringProperties>(DetailBuilder);
+	if (!GetParentTool(Properties).IsValid())
+	{
+		return;
+	}
+	const TAttribute<EVisibility> EditionVisibility = GetEditionVisibility();
+	
+	UpdateProperties(DetailBuilder, EditionVisibility, IsEnabledBySelection());
+	
+	static const FName MirrorCategoryName("Mirror");
+	IDetailCategoryBuilder& MirrorCategory = DetailBuilder.EditCategory(MirrorCategoryName, FText::GetEmpty(), ECategoryPriority::Important);
+	
+	MirrorCategory.AddCustomRow(LOCTEXT("MirrorButtonRow", "Mirror"), false)
+	[
+		SNew(SHorizontalBox)
+		+SHorizontalBox::Slot()
+		.Padding(2.f, 4.f)
+		[
+			SNew(SButton)
+			.HAlign(HAlign_Center)
+			.VAlign(VAlign_Center)
+			.Text(LOCTEXT("MirrorButtonLabel", "Mirror"))
+			.OnClicked_Lambda([Properties]()
+			{
+				Properties->MirrorBones();
+				return FReply::Handled();
+			})
+		]
+	]
+	.Visibility(EditionVisibility)
+	.IsEnabled(IsEnabledBySelection());
+}
+
+TSharedRef<IDetailCustomization> FOrientingPropertiesDetailCustomization::MakeInstance()
+{
+	return MakeShareable(new FOrientingPropertiesDetailCustomization);
+}
+
+const TArray<FName>& FOrientingPropertiesDetailCustomization::GetProperties() const
+{
+	static const TArray< FName > Properties(
+	{
+		GET_MEMBER_NAME_CHECKED(UOrientingProperties, Options),
+		GET_MEMBER_NAME_CHECKED(UOrientingProperties, bAutoOrient)
+	});
+	return Properties;
+}
+
+void FOrientingPropertiesDetailCustomization::CustomizeDetails(IDetailLayoutBuilder& DetailBuilder)
+{
+	TWeakObjectPtr<UOrientingProperties> Properties = GetCustomizedObject<UOrientingProperties>(DetailBuilder);
+	if (!GetParentTool(Properties).IsValid())
+	{
+		return;
+	}
+	const TAttribute<EVisibility> EditionVisibility = GetEditionVisibility();
+	
+	UpdateProperties(DetailBuilder, EditionVisibility, IsEnabledBySelection());
+	static const FName OrientCategoryName("Orient");
+	IDetailCategoryBuilder& MirrorCategory = DetailBuilder.EditCategory(OrientCategoryName, FText::GetEmpty(), ECategoryPriority::Important);
+	MirrorCategory.AddCustomRow(LOCTEXT("OrientButtonRow", "Orient"), false)
+	[
+		SNew(SHorizontalBox)
+		+SHorizontalBox::Slot()
+		.Padding(2.f, 4.f)
+		[
+			SNew(SButton)
+			.HAlign(HAlign_Center)
+			.VAlign(VAlign_Center)
+			.Text(LOCTEXT("OrientButtonLabel", "Orient"))
+			.IsEnabled_Lambda([Properties]()
+			{
+				return Properties->Options.Primary != EOrientAxis::None; 
+			})
+			.OnClicked_Lambda([Properties]()
+			{
+				Properties->OrientBones();
+				return FReply::Handled();
+			})
+		]
+	]
+	.Visibility(EditionVisibility)
+	.IsEnabled(IsEnabledBySelection());
+}
+
+TSharedRef<IDetailCustomization> FProjectionPropertiesDetailCustomization::MakeInstance()
+{
+	return MakeShareable(new FProjectionPropertiesDetailCustomization);
+}
+
+const TArray<FName>& FProjectionPropertiesDetailCustomization::GetProperties() const
+{
+	static const TArray< FName > Properties(
+	{
+		GET_MEMBER_NAME_CHECKED(UProjectionProperties, ProjectionType)
+	});
+	return Properties;
+}
+
+void FProjectionPropertiesDetailCustomization::CustomizeDetails(IDetailLayoutBuilder& DetailBuilder)
+{
+	const TWeakObjectPtr<UProjectionProperties> Properties = GetCustomizedObject<UProjectionProperties>(DetailBuilder);
+	if (!GetParentTool(Properties).IsValid())
+	{
+		return;
+	}
+	UpdateProperties(DetailBuilder, GetCreationVisibility());
+}
+
+namespace SkeletonEditingCustomizationLocals
+{
+	enum EOperationType
+	{
+		Add,
+		Edit,
+		Component
+	};
+}
+
+TSharedRef<IDetailCustomization> FSkeletonEditingToolDetailCustomization::MakeInstance()
+{
+	return MakeShareable(new FSkeletonEditingToolDetailCustomization);
+}
+
+void FSkeletonEditingToolDetailCustomization::CustomizeDetails(IDetailLayoutBuilder& DetailBuilder)
+{
+	TArray<TWeakObjectPtr<UObject>> ObjectsBeingCustomized;
+	DetailBuilder.GetObjectsBeingCustomized(ObjectsBeingCustomized);
+	
+	Tool = ObjectsBeingCustomized.Num() == 1 ?
+		CastChecked<USkeletonEditingTool>(ObjectsBeingCustomized[0]) : nullptr;
+	
+	if (!Tool.IsValid())
+	{
+		return;
+	}
+
+	const FSegmentedControlStyle& SegmentedControlStyle = FAppStyle::Get().GetWidgetStyle<FSegmentedControlStyle>("SegmentedControl");
+	
+	auto IsCreateEnabled = [this]()
+	{
+		return Tool->GetOperation() == EEditingOperation::Create;
+	};
+	using namespace SkeletonEditingCustomizationLocals;
+	
+	IDetailCategoryBuilder& ActionCategory = DetailBuilder.EditCategory("Action", FText::GetEmpty(), ECategoryPriority::Important);
+	ActionCategory.AddCustomRow(LOCTEXT("ActionCategory", "Action"), false)
+	[
+		SNew(SBorder)
+		.BorderImage(&SegmentedControlStyle.BackgroundBrush)
+		.Padding(FMargin(0,0,SegmentedControlStyle.UniformPadding.Right,0))
+		[
+			SNew(SHorizontalBox)
+			+SHorizontalBox::Slot()
+			[
+				SNew(SCheckBox)
+				.Clipping(EWidgetClipping::ClipToBounds)
+				.Style(&SegmentedControlStyle.FirstControlStyle)
+				.ToolTipText(LOCTEXT("AddModeTooltip", "Create new bones. (N)"))
+				.OnCheckStateChanged_Lambda([this](ECheckBoxState InCheckState)
+				{
+					if (InCheckState == ECheckBoxState::Checked)
+					{
+						Tool->SetOperation(EEditingOperation::Create);
+					}
+				})
+				.IsChecked_Lambda([IsCreateEnabled]()
+				{
+					return IsCreateEnabled() ? ECheckBoxState::Checked : ECheckBoxState::Unchecked;
+				})
+				[
+					SNew(SOverlay)
+					+ SOverlay::Slot()
+					.VAlign(VAlign_Center)
+					.HAlign(HAlign_Center)
+					[
+						SNew(STextBlock)
+						.Text(LOCTEXT("AddMode", "Add"))
+						.TextStyle(&FAppStyle::Get().GetWidgetStyle<FTextBlockStyle>("ButtonText"))
+						.ColorAndOpacity(FLinearColor::White)
+					]
+				]
+			]
+			+SHorizontalBox::Slot()
+			[
+				SNew(SCheckBox)
+				.Clipping(EWidgetClipping::ClipToBounds)
+				.Style(&SegmentedControlStyle.LastControlStyle)
+				.ToolTipText(LOCTEXT("EditModeTooltip", "Edit current bone(s) selection. (Esc)"))
+				.OnCheckStateChanged_Lambda([this](ECheckBoxState InCheckState)
+				{
+					if (InCheckState == ECheckBoxState::Checked)
+					{
+						Tool->SetOperation(EEditingOperation::Select);
+					}
+				})
+				.IsChecked_Lambda([IsCreateEnabled]()
+				{
+					return IsCreateEnabled() ? ECheckBoxState::Unchecked : ECheckBoxState::Checked;
+				})
+				[
+					SNew(SOverlay)
+						
+					+ SOverlay::Slot()
+					.VAlign(VAlign_Center)
+					.HAlign(HAlign_Center)
+					[
+						SNew(STextBlock)
+						.Text(LOCTEXT("EditMode", "Edit"))
+						.TextStyle(&FAppStyle::Get().GetWidgetStyle<FTextBlockStyle>("ButtonText"))
+						.ColorAndOpacity(FLinearColor::White)
+					]
+				]
+			]
+
+			+ SHorizontalBox::Slot()
+			.Padding(FMargin(4.f, 0.f, 0.f, 0.f))
+			.AutoWidth()
+			.HAlign(HAlign_Center)
+			.VAlign(VAlign_Center)
+			[
+				SNew(SCheckBox)
+				.Clipping(EWidgetClipping::ClipToBounds)
+				.Style(FAppStyle::Get(), "ToggleButtonCheckbox")
+				.ToolTipText(LOCTEXT("ComponentModeTooltip", "Select vertices, edges, and triangles to place bone. (T)"))
+				.OnCheckStateChanged_Lambda([this](ECheckBoxState)
+				{
+					Tool->Properties->bEnableComponentSelection = !Tool->Properties->bEnableComponentSelection;
+				})
+				.IsChecked_Lambda([this]()
+				{
+					return Tool->Properties->bEnableComponentSelection ? ECheckBoxState::Checked : ECheckBoxState::Unchecked;
+				})
+				[
+					SNew(SImage)
+						.Image(FModelingToolsEditorModeStyle::Get()->GetBrush("ModelingToolsManagerCommands.BeginMeshSelectionTool"))
+				]
+			]
+		]
+	];
+
+	CustomizeEditAction(ActionCategory);
+	CustomizeComponentSelection(DetailBuilder);
+}
+
+void FSkeletonEditingToolDetailCustomization::CustomizeComponentSelection(IDetailLayoutBuilder& DetailBuilder) const
+{
+	const TAttribute<bool> IsEnabledBySelection = TAttribute<bool>::CreateLambda([this]()
+	{
+		return Tool.IsValid() ? Tool->HasSelectedComponent() && !Tool->GetSelection().IsEmpty() : false;
+	});
+
+	const TAttribute<EVisibility> VisibilityAttribute = TAttribute<EVisibility>::CreateLambda([this]()
+	{
+		return (Tool.IsValid() && Tool->Properties->bEnableComponentSelection) ? EVisibility::Visible : EVisibility::Hidden;
+	});
+
+	static const FName SnapCategoryName("Component Snapping");
+	IDetailCategoryBuilder& SnapCategory = DetailBuilder.EditCategory(SnapCategoryName, FText::GetEmpty(), ECategoryPriority::Important);
+
+	SnapCategory.AddCustomRow(LOCTEXT("SnapButtonRow", "Snap"), false)
+	[
+		SNew(SHorizontalBox)
+		+SHorizontalBox::Slot()
+		.Padding(2.f, 4.f)
+		[
+			SNew(SButton)
+			.HAlign(HAlign_Center)
+			.VAlign(VAlign_Center)
+			.Text_Lambda([this]
+			{
+				const bool bCreate = Tool.IsValid() && Tool->GetOperation() == EEditingOperation::Create;
+				return bCreate ? LOCTEXT("SnapCreateButtonLabel", "Create") : LOCTEXT("SnapButtonLabel", "Snap");
+			})
+			.ToolTipText_Lambda([this]
+			{
+				const bool bCreate = Tool.IsValid() && Tool->GetOperation() == EEditingOperation::Create;
+				return bCreate ? LOCTEXT("SnapCreateButtonTooltip", "Create a new bone and snap it to the selected components. (V)") :
+								LOCTEXT("SnapButtonTooltip", "Snap the selected bone to the selected components. (V)");
+			})
+			.IsEnabled_Lambda([this]()
+			{
+				if (!Tool.IsValid())
+				{
+					return false;
+				}
+				const bool bCreate = Tool->GetOperation() == EEditingOperation::Create;
+				return bCreate ? Tool->HasSelectedComponent() : Tool->HasSelectedComponent() && !Tool->GetSelection().IsEmpty(); 
+			})
+			.OnClicked_Lambda([this]()
+			{
+				const bool bCreate = Tool.IsValid() && Tool->GetOperation() == EEditingOperation::Create;
+				Tool->SnapBoneToComponentSelection(bCreate);
+				return FReply::Handled();
+			})
+		]
+	]
+	.Visibility(VisibilityAttribute);
+	
+	// add selection properties
+	TObjectPtr<UMeshTopologySelectionMechanicProperties> SelectionMechanicProperties = Tool->SelectionMechanic->Properties;
+	if (IDetailPropertyRow* SelectionMechanicRow = SnapCategory.AddExternalObjects({SelectionMechanicProperties}))
+	{
+		TSharedPtr<SWidget> NameWidget, ValueWidget;
+		SelectionMechanicRow->GetDefaultWidgets(NameWidget, ValueWidget);
+		SelectionMechanicRow->CustomWidget()
+		.NameContent()
+		[
+			SNew(STextBlock)
+				.Text(FText::FromString(TEXT("Options")))
+				.Font(IDetailLayoutBuilder::GetDetailFont())
+		]
+		.ValueContent()
+		[
+			ValueWidget->AsShared()
+		];
+		SelectionMechanicRow->ShouldAutoExpand(false);
+		SelectionMechanicRow->Visibility(VisibilityAttribute);
+	}
+}
+
+void FSkeletonEditingToolDetailCustomization::CustomizeEditAction(IDetailCategoryBuilder& InActionCategory) const
+{
+	if (!Tool.IsValid())
+	{
+		return;
+	}
+	
+	const TAttribute<EVisibility> EditionVisibility = TAttribute<EVisibility>::CreateLambda([this]()
+	{
+		if (!Tool.IsValid())
+		{
+			return EVisibility::Hidden;
+		}
+		return Tool->GetOperation() == EEditingOperation::Create ? EVisibility::Hidden : EVisibility::Visible;
+	});
+	
+	const TAttribute<bool> IsEnabledBySelection = TAttribute<bool>::CreateLambda([this]()
+	{
+		return Tool.IsValid() ? !Tool->GetSelection().IsEmpty() : false;
+	});
+	FSegmentedControlStyle SegmentedControlStyle = FAppStyle::Get().GetWidgetStyle<FSegmentedControlStyle>("SegmentedControl");
+	
+	InActionCategory.AddCustomRow(LOCTEXT("EditCategory", "Edit"), false)
+	[
+		SNew(SHorizontalBox)
+		+SHorizontalBox::Slot()
+		.Padding(2.f, 4.f)
+		[
+			SNew(SCheckBox)
+			.Style(FAppStyle::Get(), "SegmentedCombo.ButtonOnly")
+			.OnCheckStateChanged_Lambda([this](ECheckBoxState InCheckState)
+			{
+				const EEditingOperation Operation = InCheckState == ECheckBoxState::Checked ?
+					EEditingOperation::Parent : EEditingOperation::Select;
+				Tool->SetOperation(Operation);
+			})
+			.IsChecked_Lambda([this]()
+			{
+				return Tool->GetOperation() == EEditingOperation::Parent ? ECheckBoxState::Checked : ECheckBoxState::Unchecked;
+			})
+			.IsEnabled(IsEnabledBySelection)
+			[
+				SNew(SOverlay)
+				
+				+ SOverlay::Slot()
+				.VAlign(VAlign_Center)
+				.HAlign(HAlign_Center)
+				[
+					SNew(STextBlock)
+					.Text(LOCTEXT("EditParent", "Reparent"))
+					.ToolTipText(LOCTEXT("EditParentTooltip", "Click on a bone to be set as the new parent. (B)"))
+					.TextStyle(&FAppStyle::Get().GetWidgetStyle<FTextBlockStyle>("ButtonText"))
+					.ColorAndOpacity(FLinearColor::White)
+				]
+			]
+		]
+		+SHorizontalBox::Slot()
+		.Padding(2.f, 4.f)
+		[
+			SNew(SButton)
+			.HAlign(HAlign_Center)
+			.VAlign(VAlign_Center)
+			.Text(LOCTEXT("EditDisconnect", "Disconnect"))
+			.ToolTipText(LOCTEXT("EditDisconnectTooltip", "Unparent current selection. (SHIFT+P)"))
+			.OnClicked_Lambda([this]()
+			{
+				Tool->UnParentBones();
+				return FReply::Handled();
+			})
+			.IsEnabled_Lambda([this]()
+			{
+				if (!Tool.IsValid())
+				{
+					return false;
+				}
+				return !Tool->GetSelection().IsEmpty() && Tool->GetOperation() != EEditingOperation::Parent;
+			})
+		]
+		
+		+SHorizontalBox::Slot()
+		.Padding(2.f, 4.f)
+		[
+			SNew(SButton)
+			.HAlign(HAlign_Center)
+			.VAlign(VAlign_Center)
+			.Text(LOCTEXT("EditRemove", "Remove"))
+			.ToolTipText(LOCTEXT("EditRemoveTooltip", "Remove current selection. (Delete)"))
+			.OnClicked_Lambda([this]()
+			{
+				Tool->RemoveBones();
+				return FReply::Handled();
+			})
+			.IsEnabled_Lambda([this]()
+			{
+				if (!Tool.IsValid())
+				{
+					return false;
+				}
+				return !Tool->GetSelection().IsEmpty() && Tool->GetOperation() != EEditingOperation::Parent;
+			})
+		]
+	]
+	.Visibility(EditionVisibility);
+}
+
+#undef LOCTEXT_NAMESPACE
\ No newline at end of file
diff --git a/Engine/Plugins/Animation/SkeletalMeshModelingTools/Source/SkeletalMeshModelingTools/Private/DetailCustomization/SkeletonEditingToolPropertyCustomizations.h b/Engine/Plugins/Animation/SkeletalMeshModelingTools/Source/SkeletalMeshModelingTools/Private/DetailCustomization/SkeletonEditingToolPropertyCustomizations.h
new file mode 100644
--- /dev/null
+++ b/Engine/Plugins/Animation/SkeletalMeshModelingTools/Source/SkeletalMeshModelingTools/Private/DetailCustomization/SkeletonEditingToolPropertyCustomizations.h
@@ -0,0 +1,131 @@
+// Copyright Epic Games, Inc. All Rights Reserved.
+
+#pragma once
+
+#include "IDetailCustomization.h"
+
+#include "SAdvancedTransformInputBox.h"
+#include "SkeletalMesh/SkeletonEditingTool.h"
+
+class IDetailLayoutBuilder;
+
+/**
+ * FSkeletonEditingToolDetailCustomization
+ */
+
+class FSkeletonEditingToolDetailCustomization : public IDetailCustomization
+{
+public:
+	static TSharedRef<IDetailCustomization> MakeInstance();
+	virtual void CustomizeDetails(IDetailLayoutBuilder& DetailBuilder) override;
+private:
+	void CustomizeEditAction(IDetailCategoryBuilder& InActionCategory) const;
+	void CustomizeComponentSelection(IDetailLayoutBuilder& DetailBuilder) const;
+	
+	TWeakObjectPtr<USkeletonEditingTool> Tool = nullptr;
+};
+
+/**
+ * ISkeletonEditingPropertiesDetailCustomization
+ */
+
+class ISkeletonEditingPropertiesDetailCustomization : public IDetailCustomization
+{
+protected:
+	template< typename T >
+	static TWeakObjectPtr<T> GetCustomizedObject(const IDetailLayoutBuilder& InDetailBuilder)
+	{
+		TArray<TWeakObjectPtr<UObject>> ObjectsBeingCustomized;
+		InDetailBuilder.GetObjectsBeingCustomized(ObjectsBeingCustomized);
+		if (ObjectsBeingCustomized.Num() != 1)
+		{
+			return nullptr;
+		}
+		return CastChecked<T>(ObjectsBeingCustomized[0]);
+	}
+	
+	template<typename TProperty>
+	TWeakObjectPtr<USkeletonEditingTool> GetParentTool(const TWeakObjectPtr<TProperty>& InProperties)
+	{
+		Tool = InProperties.IsValid() ? InProperties->ParentTool : nullptr;	
+		return Tool;
+	}
+
+	TAttribute<EVisibility> GetCreationVisibility() const;
+	TAttribute<EVisibility> GetEditionVisibility() const;
+	TAttribute<bool> IsEnabledBySelection() const;
+	
+	void UpdateProperties(
+		IDetailLayoutBuilder& InDetailBuilder,
+		const TAttribute<EVisibility>& InVisibility,
+		const TAttribute<bool>& InEnabled = TAttribute<bool>()) const;
+	
+	virtual const TArray<FName>& GetProperties() const = 0;
+	
+	TWeakObjectPtr<USkeletonEditingTool> Tool;
+};
+
+/**
+ * FSkeletonEditingPropertiesDetailCustomization
+ */
+
+class FSkeletonEditingPropertiesDetailCustomization : public ISkeletonEditingPropertiesDetailCustomization
+{
+public:
+	static TSharedRef<IDetailCustomization> MakeInstance();
+	virtual void CustomizeDetails(IDetailLayoutBuilder& DetailBuilder) override;
+
+protected:
+	virtual const TArray<FName>& GetProperties() const override;
+	
+private:
+	void CustomizeValueGet(SAdvancedTransformInputBox<FTransform>::FArguments& InOutArgs);
+	void CustomizeValueSet(SAdvancedTransformInputBox<FTransform>::FArguments& InOutArgs);
+	void CustomizeClipboard(SAdvancedTransformInputBox<FTransform>::FArguments& InOutArgs);
+	
+	TUniquePtr<SkeletonEditingTool::FRefSkeletonChange> ActiveChange;
+	TBitArray<> RelativeArray;
+};
+
+/**
+ * FMirroringPropertiesDetailCustomization
+ */
+
+class FMirroringPropertiesDetailCustomization : public ISkeletonEditingPropertiesDetailCustomization
+{
+public:
+	static TSharedRef<IDetailCustomization> MakeInstance();
+	virtual void CustomizeDetails(IDetailLayoutBuilder& DetailBuilder) override;
+
+protected:
+	virtual const TArray<FName>& GetProperties() const override;
+};
+
+/**
+ * FOrientingPropertiesDetailCustomization
+ */
+
+class FOrientingPropertiesDetailCustomization : public ISkeletonEditingPropertiesDetailCustomization
+{
+public:
+	static TSharedRef<IDetailCustomization> MakeInstance();
+	virtual void CustomizeDetails(IDetailLayoutBuilder& DetailBuilder) override;
+
+protected:
+	virtual const TArray<FName>& GetProperties() const override;
+};
+
+/**
+ * FProjectionPropertiesDetailCustomization
+ */
+
+class FProjectionPropertiesDetailCustomization : public ISkeletonEditingPropertiesDetailCustomization
+{
+public:
+	static TSharedRef<IDetailCustomization> MakeInstance();
+	virtual void CustomizeDetails(IDetailLayoutBuilder& DetailBuilder) override;
+	
+protected:
+	virtual const TArray<FName>& GetProperties() const override;
+};
+
diff --git a/Engine/Plugins/Animation/SkeletalMeshModelingTools/Source/SkeletalMeshModelingTools/Private/DetailCustomization/SkinWeightDetailCustomization.cpp b/Engine/Plugins/Animation/SkeletalMeshModelingTools/Source/SkeletalMeshModelingTools/Private/DetailCustomization/SkinWeightDetailCustomization.cpp
new file mode 100644
--- /dev/null
+++ b/Engine/Plugins/Animation/SkeletalMeshModelingTools/Source/SkeletalMeshModelingTools/Private/DetailCustomization/SkinWeightDetailCustomization.cpp
@@ -0,0 +1,1552 @@
+﻿// Copyright Epic Games, Inc. All Rights Reserved.
+
+#include "SkinWeightDetailCustomization.h"
+
+#include "DetailCategoryBuilder.h"
+#include "DetailLayoutBuilder.h"
+#include "DetailWidgetRow.h"
+#include "ModelingToolsEditorModeStyle.h"
+#include "Widgets/Input/SSegmentedControl.h"
+#include "SkeletalMesh/SkinWeightsPaintTool.h"
+#include "Widgets/Input/SSpinBox.h"
+#include "Widgets/Input/SButton.h"
+#include "Widgets/Input/SNumericEntryBox.h"
+#include "Framework/MultiBox/MultiBoxBuilder.h"
+
+#define LOCTEXT_NAMESPACE "SkinWeightToolSettingsEditor"
+
+// layout constants
+float FSkinWeightDetailCustomization::WeightSliderWidths = 150.0f;
+float FSkinWeightDetailCustomization::WeightEditingLabelsPercent = 0.40f;
+float FSkinWeightDetailCustomization::WeightEditVerticalPadding = 4.0f;
+float FSkinWeightDetailCustomization::WeightEditHorizontalPadding = 2.0f;
+
+static FName ColumnName_Bone = "Bone";
+static FName ColumnName_Weight = "Weight";
+static FName ColumnName_Prune = "Prune";
+
+FSkinWeightDetailCustomization::~FSkinWeightDetailCustomization()
+{
+	if (Tool.IsValid())
+	{
+		Tool->OnSelectionChanged.RemoveAll(this);
+	}
+
+	Tool.Reset();
+	ToolSettings.Reset();
+}
+
+void FSkinWeightDetailCustomization::CustomizeDetails(IDetailLayoutBuilder& DetailBuilder)
+{
+	CurrentDetailBuilder = &DetailBuilder;
+	
+	TArray<TWeakObjectPtr<UObject>> DetailObjects;
+	DetailBuilder.GetObjectsBeingCustomized(DetailObjects);
+
+	// should be impossible to get multiple settings objects for a single tool
+	ensure(DetailObjects.Num()==1);
+	ToolSettings = Cast<USkinWeightsPaintToolProperties>(DetailObjects[0]);
+	Tool = ToolSettings->WeightTool;
+	Tool->OnSelectionChanged.AddSP(this, &FSkinWeightDetailCustomization::OnSelectionChanged);
+	
+	// custom display of falloff mode as segmented toggle buttons
+	IDetailCategoryBuilder& EditModeCategory = DetailBuilder.EditCategory("Weight Editing Mode", FText::GetEmpty(), ECategoryPriority::Important);
+
+	// add segmented control toggle for editing modes ("Brush" or "Selection")
+	EditModeCategory.AddCustomRow(LOCTEXT("EditModeCategory", "Weight Editing Mode"), false)
+	.WholeRowContent()
+	[
+		SNew(SBox)
+		.Padding(2.0f)
+		[
+			SNew(SSegmentedControl<EWeightEditMode>)
+			.ToolTipText(LOCTEXT("EditingModeTooltip",
+					"Brush: edit weights by painting on mesh.\n"
+					"Mesh: select vertices/edges/faces to edit weights directly.\n"
+					"Bones: select and manipulate bones to preview deformations.\n"))
+			.Value_Lambda([this]()
+			{
+				return ToolSettings->EditingMode;
+			})
+			.OnValueChanged_Lambda([this](EWeightEditMode Mode)
+			{
+				ToolSettings->EditingMode = Mode;
+				ToolSettings->WeightTool->ToggleEditingMode();
+				if (CurrentDetailBuilder)
+				{
+					CurrentDetailBuilder->ForceRefreshDetails();
+				}
+			})
+			+SSegmentedControl<EWeightEditMode>::Slot(EWeightEditMode::Brush)
+			.Text(LOCTEXT("BrushEditMode", "Brush"))
+			+ SSegmentedControl<EWeightEditMode>::Slot(EWeightEditMode::Mesh)
+			.Text(LOCTEXT("MeshEditMode", "Mesh"))
+			+ SSegmentedControl<EWeightEditMode>::Slot(EWeightEditMode::Bones)
+			.Text(LOCTEXT("BoneEditMode", "Bones"))
+		]
+	];
+
+	// BRUSH editing mode UI
+	if (ToolSettings->EditingMode == EWeightEditMode::Brush)
+	{
+		AddBrushUI(DetailBuilder);
+	}
+
+	// MESH editing mode UI
+	if (ToolSettings->EditingMode == EWeightEditMode::Mesh)
+	{
+		AddSelectionUI(DetailBuilder);
+	}
+	
+	// COLOR MODE category
+	IDetailCategoryBuilder& MeshDisplayCategory = DetailBuilder.EditCategory("MeshDisplay", FText::GetEmpty(), ECategoryPriority::Important);
+	MeshDisplayCategory.InitiallyCollapsed(false);
+	MeshDisplayCategory.AddCustomRow(LOCTEXT("ColorModeCategory", "Color Mode"), false)
+	.WholeRowContent()
+	[
+		SNew(SBox)
+		.Padding(2.0f)
+		.HAlign(HAlign_Center)
+		[
+			SNew(SSegmentedControl<EWeightColorMode>)
+			.ToolTipText(LOCTEXT("WeightColorTooltip",
+					"Adjust the weight display in the viewport.\n\n"
+					"Greyscale: Displays weights on the current bone by blending from black (0) to white (1).\n"
+					"Ramp: Displays weights on the current bone. Weights at 0 and 1 use the min and max colors. Weights inbetween 0 and 1 use the ramp colors.\n"
+					"Multi Color: Displays weights on ALL bones using the color of the bones.\n"
+					"Full Material: Displays normal mesh materials with textures.\n"))
+			.Value_Lambda([this]()
+			{
+				return ToolSettings->ColorMode;
+			})
+			.OnValueChanged_Lambda([this](EWeightColorMode Mode)
+			{
+				ToolSettings->SetColorMode(Mode);
+			})
+			+ SSegmentedControl<EWeightColorMode>::Slot(EWeightColorMode::Greyscale)
+			.Text(LOCTEXT("GreyscaleMode", "Greyscale"))
+			+SSegmentedControl<EWeightColorMode>::Slot(EWeightColorMode::Ramp)
+			.Text(LOCTEXT("RampMode", "Ramp"))
+			+SSegmentedControl<EWeightColorMode>::Slot(EWeightColorMode::BoneColors)
+			.Text(LOCTEXT("BoneColorsMode", "Bone Colors"))
+			+SSegmentedControl<EWeightColorMode>::Slot(EWeightColorMode::FullMaterial)
+			.Text(LOCTEXT("MaterialMode", "Full Material"))
+		]
+	];
+
+	AddTransferUI(DetailBuilder);
+
+	// Edit SkinWeightLayer category 
+	IDetailCategoryBuilder& SkinWeightLayerCategory = DetailBuilder.EditCategory("SkinWeightLayer", FText::GetEmpty(), ECategoryPriority::Important);
+	SkinWeightLayerCategory.InitiallyCollapsed(true);
+	
+	// hide skin weight tool properties that were customized
+	TArray<FName> PropertiesToHide;
+	PropertiesToHide.Add(GET_MEMBER_NAME_CHECKED(USkinWeightsPaintToolProperties, BrushMode));
+	
+	PropertiesToHide.Add(GET_MEMBER_NAME_CHECKED(USkinWeightsPaintToolProperties, EditingMode));
+	PropertiesToHide.Add(GET_MEMBER_NAME_CHECKED(USkinWeightsPaintToolProperties, ColorMode));
+	PropertiesToHide.Add(GET_MEMBER_NAME_CHECKED(USkinWeightsPaintToolProperties, MeshSelectMode));
+	PropertiesToHide.Add(GET_MEMBER_NAME_CHECKED(USkinWeightsPaintToolProperties, SourceSkeletalMesh));
+	PropertiesToHide.Add(GET_MEMBER_NAME_CHECKED(USkinWeightsPaintToolProperties, SourceLOD));
+	PropertiesToHide.Add(GET_MEMBER_NAME_CHECKED(USkinWeightsPaintToolProperties, SourceSkinWeightProfile));
+	for (const FName PropertyToHide : PropertiesToHide)
+	{
+		TSharedRef<IPropertyHandle> Property = DetailBuilder.GetProperty(PropertyToHide);
+		DetailBuilder.HideProperty(Property);
+	}
+	// hide base class properties
+	PropertiesToHide.Reset();
+	PropertiesToHide.Add(GET_MEMBER_NAME_CHECKED(UBrushBaseProperties, bSpecifyRadius));
+	PropertiesToHide.Add(GET_MEMBER_NAME_CHECKED(UBrushBaseProperties, BrushSize));
+	PropertiesToHide.Add(GET_MEMBER_NAME_CHECKED(UBrushBaseProperties, BrushStrength));
+	PropertiesToHide.Add(GET_MEMBER_NAME_CHECKED(UBrushBaseProperties, BrushFalloffAmount));
+	PropertiesToHide.Add(GET_MEMBER_NAME_CHECKED(UBrushBaseProperties, BrushRadius));
+	for (const FName PropertyToHide : PropertiesToHide)
+	{
+		TSharedRef<IPropertyHandle> Property = DetailBuilder.GetProperty(PropertyToHide, UBrushBaseProperties::StaticClass());
+		DetailBuilder.HideProperty(Property);
+	}
+}
+
+void FSkinWeightDetailCustomization::AddBrushUI(IDetailLayoutBuilder& DetailBuilder) const
+{
+	// custom display of falloff mode as segmented toggle buttons
+	IDetailCategoryBuilder& BrushCategory = DetailBuilder.EditCategory("Brush", FText::GetEmpty(), ECategoryPriority::Important);
+
+	// add segmented control toggle for brush behavior modes ("Add", "Replace", etc..)
+	BrushCategory.AddCustomRow(LOCTEXT("BrushModeCategory", "Brush Mode"), false)
+	.WholeRowContent()
+	[
+		SNew(SBox)
+		.Padding(2.0f)
+		[
+			SNew(SSegmentedControl<EWeightEditOperation>)
+			.ToolTipText(LOCTEXT("BrushTooltip",
+				"Select the operation to apply when using the brush.\n"
+				"Add: applies the current weight plus the Strength value to the new weight.\n"
+				"Replace: applies the current weight minus the Strength value to the new weight.\n"
+				"Multiply: applies the current weight multiplied by the Strength value to the new weight.\n"
+				"Relax: applies the average of the connected (by edge) vertex weights to the new vertex weight, blended by the Strength."))
+			.Value_Lambda([this]()
+			{
+				return ToolSettings->BrushMode;
+			})
+			.OnValueChanged_Lambda([this](EWeightEditOperation Mode)
+			{
+				ToolSettings->SetBrushMode(Mode);
+			})
+			+SSegmentedControl<EWeightEditOperation>::Slot(EWeightEditOperation::Add)
+			.Text(LOCTEXT("BrushAddMode", "Add"))
+			+ SSegmentedControl<EWeightEditOperation>::Slot(EWeightEditOperation::Replace)
+			.Text(LOCTEXT("BrushReplaceMode", "Replace"))
+			+ SSegmentedControl<EWeightEditOperation>::Slot(EWeightEditOperation::Multiply)
+			.Text(LOCTEXT("BrushMultiplyMode", "Multiply"))
+			+ SSegmentedControl<EWeightEditOperation>::Slot(EWeightEditOperation::Relax)
+			.Text(LOCTEXT("BrushRelaxMode", "Relax"))
+		]
+	];
+
+	// add segmented control toggle for brush falloff modes ("Surface" or "Volume")
+	BrushCategory.AddCustomRow(LOCTEXT("BrushFalloffModeCategory", "Brush Falloff Mode"), false)
+	.WholeRowContent()
+	[
+		SNew(SBox)
+		.Padding(2.0f)
+		[
+			SNew(SSegmentedControl<EWeightBrushFalloffMode>)
+			.ToolTipText(LOCTEXT("BrushFalloffModeTooltip",
+					"Surface: falloff is based on the distance along the surface from the brush center to nearby connected vertices.\n"
+					"Volume: falloff is based on the straight-line distance from the brush center to surrounding vertices.\n"))
+			.Value_Lambda([this]()
+			{
+				return ToolSettings->GetBrushConfig().FalloffMode;
+			})
+			.OnValueChanged_Lambda([this](EWeightBrushFalloffMode Mode)
+			{
+				ToolSettings->SetFalloffMode(Mode);
+			})
+			+SSegmentedControl<EWeightBrushFalloffMode>::Slot(EWeightBrushFalloffMode::Surface)
+			.Text(LOCTEXT("SurfaceMode", "Surface"))
+			+ SSegmentedControl<EWeightBrushFalloffMode>::Slot(EWeightBrushFalloffMode::Volume)
+			.Text(LOCTEXT("VolumeMode", "Volume"))
+		]
+	];
+
+	BrushCategory.AddCustomRow(LOCTEXT("BrushSizeCategory", "Brush Radius"), false)
+	.NameContent()
+	[
+		SNew(STextBlock)
+		.Text(LOCTEXT("BrushRadiusLabel", "Radius"))
+		.Font(FAppStyle::GetFontStyle(TEXT("PropertyWindow.NormalFont")))
+		.ToolTipText(LOCTEXT("BrushRadiusTooltip", "The radius of the brush in scene units."))
+	]
+	.ValueContent()
+	[
+		SNew(SSpinBox<float>)
+		.MinValue(0.01f)
+		.MaxSliderValue(20.f)
+		.Value(10.0f)
+		.SupportDynamicSliderMaxValue(true)
+		.Value_Lambda([this]()
+		{
+			return ToolSettings->GetBrushConfig().Radius;
+		})
+		.OnValueChanged_Lambda([this](float NewValue)
+		{
+			ToolSettings->BrushRadius = NewValue;
+			ToolSettings->GetBrushConfig().Radius = NewValue;
+			FPropertyChangedEvent PropertyChangedEvent(UBrushBaseProperties::StaticClass()->FindPropertyByName(GET_MEMBER_NAME_CHECKED(UBrushBaseProperties, BrushRadius)));
+			ToolSettings->PostEditChangeProperty(PropertyChangedEvent);
+		})
+		.OnValueCommitted_Lambda([this](float NewValue, ETextCommit::Type CommitType)
+		{
+			ToolSettings->SaveConfig();
+		})
+	];
+
+	BrushCategory.AddCustomRow(LOCTEXT("BrushStrengthCategory", "Brush Strength"), false)
+	.NameContent()
+	[
+		SNew(STextBlock)
+		.Text(LOCTEXT("BrushStrengthLabel", "Strength"))
+		.Font(FAppStyle::GetFontStyle(TEXT("PropertyWindow.NormalFont")))
+		.ToolTipText(LOCTEXT("BrushStrengthTooltip", "The strength of the effect on the weights. Exact effect depends on the Brush mode."))
+	]
+	.ValueContent()
+	[
+		SNew(SSpinBox<float>)
+		.MinValue(0.f)
+		.MaxValue(2.0f)
+		.MaxSliderValue(1.f)
+		.Value(1.0f)
+		.SupportDynamicSliderMaxValue(true)
+		.Value_Lambda([this]()
+		{
+			return ToolSettings->GetBrushConfig().Strength;
+		})
+		.OnValueChanged_Lambda([this](float NewValue)
+		{
+			ToolSettings->BrushStrength = NewValue;
+			ToolSettings->GetBrushConfig().Strength = NewValue;
+			FPropertyChangedEvent PropertyChangedEvent(UBrushBaseProperties::StaticClass()->FindPropertyByName(GET_MEMBER_NAME_CHECKED(UBrushBaseProperties, BrushStrength)));
+			ToolSettings->PostEditChangeProperty(PropertyChangedEvent);
+		})
+		.OnValueCommitted_Lambda([this](float NewValue, ETextCommit::Type CommitType)
+		{
+			ToolSettings->SaveConfig();
+		})
+	];
+
+	BrushCategory.AddCustomRow(LOCTEXT("BrushFalloffCategory", "Brush Falloff"), false)
+	.NameContent()
+	[
+		SNew(STextBlock)
+		.Text(LOCTEXT("BrushFalloffLabel", "Falloff"))
+		.Font(FAppStyle::GetFontStyle(TEXT("PropertyWindow.NormalFont")))
+		.ToolTipText(LOCTEXT("BrushFalloffTooltip", "At 0, the brush has no falloff. At 1 it has exponential falloff."))
+	]
+	.ValueContent()
+	[
+		SNew(SSpinBox<float>)
+		.MinValue(0.f)
+		.MaxValue(1.f)
+		.Value_Lambda([this]()
+		{
+			return ToolSettings->GetBrushConfig().Falloff;
+		})
+		.OnValueChanged_Lambda([this](float NewValue)
+		{
+			ToolSettings->BrushFalloffAmount = NewValue;
+			ToolSettings->GetBrushConfig().Falloff = NewValue;
+			FPropertyChangedEvent PropertyChangedEvent(UBrushBaseProperties::StaticClass()->FindPropertyByName(GET_MEMBER_NAME_CHECKED(UBrushBaseProperties, BrushFalloffAmount)));
+			ToolSettings->PostEditChangeProperty(PropertyChangedEvent);
+		})
+		.OnValueCommitted_Lambda([this](float NewValue, ETextCommit::Type CommitType)
+		{
+			ToolSettings->SaveConfig();
+		})
+	];
+}
+
+void FSkinWeightDetailCustomization::AddSelectionUI(IDetailLayoutBuilder& DetailBuilder) const
+{
+	// custom display of selection editing tools
+	IDetailCategoryBuilder& EditSelectionCategory = DetailBuilder.EditCategory("Edit Selection", FText::GetEmpty(), ECategoryPriority::Important);
+	EditSelectionCategory.InitiallyCollapsed(true);
+
+	// create a toolbar for the selection filter
+	FSlimHorizontalToolBarBuilder ToolbarBuilder(MakeShared<FUICommandList>(), FMultiBoxCustomization::None);
+	ToolbarBuilder.SetStyle(FModelingToolsEditorModeStyle::Get().Get(), "PolyEd.SelectionToolbar");
+	ToolbarBuilder.SetLabelVisibility(EVisibility::Collapsed);
+
+	ToolbarBuilder.BeginSection("SelectionFilter");
+	ToolbarBuilder.BeginBlockGroup();
+	
+	auto AddToggleButtonForBool = [&ToolbarBuilder, this](EComponentSelectionMode Mode, const FText& Label, const FText& Tooltip, const FName IconName)
+	{
+		ToolbarBuilder.AddToolBarButton(FUIAction(
+		FExecuteAction::CreateLambda([this, Mode]()
+		{
+			ToolSettings->SetComponentMode(Mode);
+		}),
+		FCanExecuteAction::CreateLambda([this]()
+		{
+			return ToolSettings->EditingMode == EWeightEditMode::Mesh;
+		}),
+		FIsActionChecked::CreateLambda([this, Mode]()
+		{
+			return ToolSettings->ComponentSelectionMode == Mode;
+		})),
+		NAME_None,	// Extension hook
+		Label,		// Label
+		Tooltip,	// Tooltip
+		FSlateIcon(FModelingToolsEditorModeStyle::Get()->GetStyleSetName(), IconName),
+		EUserInterfaceActionType::ToggleButton);
+	};
+
+	AddToggleButtonForBool(
+		EComponentSelectionMode::Vertices,
+		LOCTEXT("VerticesLabel", "Vertices"),
+		LOCTEXT("VerticesTooltip", "Select mesh vertices."),
+		"PolyEd.SelectCorners");
+	AddToggleButtonForBool(
+		EComponentSelectionMode::Edges,
+		LOCTEXT("EdgesLabel", "Edges"),
+		LOCTEXT("EdgesTooltip", "Select mesh edges."),
+		"PolyEd.SelectEdges");
+	AddToggleButtonForBool(
+		EComponentSelectionMode::Faces,
+		LOCTEXT("FacesLabel", "Faces"),
+		LOCTEXT("FacesTooltip", "Select mesh faces."),
+		"PolyEd.SelectFaces");
+
+	ToolbarBuilder.EndBlockGroup();
+	ToolbarBuilder.EndSection();
+
+	// edit selection category
+	EditSelectionCategory.AddCustomRow(LOCTEXT("EditSelectionRow", "Edit Selection"), false)
+	.WholeRowContent()
+	[
+		SNew(SVerticalBox)
+		+SVerticalBox::Slot()
+		.AutoHeight()
+		[
+			SNew(SHorizontalBox)
+
+			// ISOLATE SELECTION
+			+SHorizontalBox::Slot()
+			.HAlign(HAlign_Center)
+			.VAlign(VAlign_Center)
+			.Padding(2.0f, WeightEditVerticalPadding)
+			[
+				SNew(SHorizontalBox)
+				+SHorizontalBox::Slot()
+				[
+					ToolbarBuilder.MakeWidget()
+				]
+
+				+SHorizontalBox::Slot()
+				.AutoWidth()
+				.VAlign(VAlign_Center)
+				[
+					SNew(SCheckBox)
+					.Style(FAppStyle::Get(), "ToggleButtonCheckbox")
+					.HAlign(HAlign_Center)
+					.ToolTipText(LOCTEXT("IsolateSelectedTooltip",
+							"Shows only the selected faces in the viewport.\n"
+							"Weight editing operations will not affect hidden vertices.\n"
+							"NOTE: This only works on the target (main) mesh."))
+					.IsEnabled_Lambda([this]()
+					{
+						// isolated selection only available on main mesh (for now)
+						const bool bHasSelection = Tool->GetMainMeshSelector()->IsAnyComponentSelected();
+						const bool bAlreadyIsolatingSelection = Tool->GetSelectionIsolator()->IsSelectionIsolated();
+						return bHasSelection ||  bAlreadyIsolatingSelection;
+					})
+					.IsChecked_Lambda([this]()
+					{
+						const bool bAlreadyIsolatingSelection = Tool->GetSelectionIsolator()->IsSelectionIsolated();
+						return bAlreadyIsolatingSelection ? ECheckBoxState::Checked : ECheckBoxState::Unchecked;
+					})
+					.OnCheckStateChanged_Lambda([this](ECheckBoxState InCheckBoxState)
+					{
+						if (InCheckBoxState == ECheckBoxState::Checked)
+						{
+							Tool->GetSelectionIsolator()->IsolateSelectionAsTransaction();
+						}
+						else
+						{
+							Tool->GetSelectionIsolator()->UnIsolateSelectionAsTransaction();
+						}
+					})
+					[
+						SNew(STextBlock)
+						.Text_Lambda([this]()
+						{
+							if (Tool->GetSelectionIsolator()->IsSelectionIsolated())
+							{
+								return LOCTEXT("ShowAllButtonLabel", "Show Full Mesh");
+							}
+								
+							return LOCTEXT("IsolateButtonLabel", "Isolate Selected");
+						})
+					]
+				]
+			]
+		]
+
+		// GROW / SHRINK / FLOOD
+		+SVerticalBox::Slot()
+		[
+			SNew(SHorizontalBox)
+		
+			+SHorizontalBox::Slot()
+			.Padding(2.f, WeightEditVerticalPadding)
+			[
+				SNew(SButton)
+				.IsEnabled_Lambda([this]{ return Tool->GetActiveMeshSelector()->IsAnyComponentSelected(); })
+				.HAlign(HAlign_Center)
+				.VAlign(VAlign_Center)
+				.Text(LOCTEXT("GrowSelectionButtonLabel", "Grow"))
+				.ToolTipText(LOCTEXT("GrowSelectionTooltip",
+						"Grow the current selection by adding connected neighbors to current selection.\n"))
+				.OnClicked_Lambda([this]()
+				{
+					Tool->GetActiveMeshSelector()->GrowSelection();
+					return FReply::Handled();
+				})
+			]
+
+			+SHorizontalBox::Slot()
+			.Padding(2.f, WeightEditVerticalPadding)
+			[
+				SNew(SButton)
+				.IsEnabled_Lambda([this]{ return Tool->GetActiveMeshSelector()->IsAnyComponentSelected(); })
+				.HAlign(HAlign_Center)
+				.VAlign(VAlign_Center)
+				.Text(LOCTEXT("ShrinkSelectionButtonLabel", "Shrink"))
+				.ToolTipText(LOCTEXT("ShrinkSelectionTooltip",
+						"Shrink the current selection by removing components on the border of the current selection.\n"))
+				.OnClicked_Lambda([this]()
+				{
+					Tool->GetActiveMeshSelector()->ShrinkSelection();
+					return FReply::Handled();
+				})
+			]
+
+			+SHorizontalBox::Slot()
+			.Padding(2.f, WeightEditVerticalPadding)
+			[
+				SNew(SButton)
+				.IsEnabled_Lambda([this]{ return Tool->GetActiveMeshSelector()->IsAnyComponentSelected(); })
+				.HAlign(HAlign_Center)
+				.VAlign(VAlign_Center)
+				.Text(LOCTEXT("FloodSelectionButtonLabel", "Flood"))
+				.ToolTipText(LOCTEXT("FloodSelectionTooltip",
+						"Flood the current selection by adding all connected components to the current selection.\n"))
+				.OnClicked_Lambda([this]()
+				{
+					Tool->GetActiveMeshSelector()->FloodSelection();
+					return FReply::Handled();
+				})
+			]
+		]
+
+		// SELECT AFFECTED VERTICES
+		+SVerticalBox::Slot()
+		[
+			SNew(SHorizontalBox)
+		
+			+SHorizontalBox::Slot()
+			.Padding(2.f, WeightEditVerticalPadding)
+			[
+				SNew(SButton)
+				.IsEnabled_Lambda([this]
+				{
+					// only allow selecting affected vertices on the target/main mesh
+					return ToolSettings->MeshSelectMode == EMeshTransferOption::Target;
+				})
+				.HAlign(HAlign_Center)
+				.VAlign(VAlign_Center)
+				.Text(LOCTEXT("AffectedSelectionButtonLabel", "Affected"))
+				.ToolTipText(LOCTEXT("AffectedSelectionTooltip",
+						"Select vertices that are affected by the currently selected bone(s).\n"
+						"Holding Shift or Ctrl will add or subtract affected vertices from the current selection."))
+				.OnClicked_Lambda([this]()
+				{
+					Tool->SelectAffected();
+					return FReply::Handled();
+				})
+			]
+
+			+SHorizontalBox::Slot()
+			.Padding(2.f, WeightEditVerticalPadding)
+			[
+				SNew(SButton)
+				.IsEnabled_Lambda([this]{ return Tool->GetActiveMeshSelector()->IsAnyComponentSelected(); })
+				.HAlign(HAlign_Center)
+				.VAlign(VAlign_Center)
+				.Text(LOCTEXT("BorderSelectionButtonLabel", "Convert to Border"))
+				.ToolTipText(LOCTEXT("BorderSelectionTooltip", "Select vertices on the border of the current selection."))
+				.OnClicked_Lambda([this]()
+				{
+					Tool->GetActiveMeshSelector()->SelectBorder();
+					return FReply::Handled();
+				})
+			]
+		]
+
+		// SELECT BY INFLUENCE COUNT
+		+SVerticalBox::Slot()
+		[
+			SNew(SHorizontalBox)
+		
+			+SHorizontalBox::Slot()
+			.Padding(2.f, WeightEditVerticalPadding)
+			[
+				SNew(SButton)
+				.IsEnabled_Lambda([this]
+				{
+					// only allow selecting affected vertices on the target/main mesh
+					return ToolSettings->MeshSelectMode == EMeshTransferOption::Target;
+				})
+				.HAlign(HAlign_Center)
+				.VAlign(VAlign_Center)
+				.Text(LOCTEXT("InfluenceCountSelectionButtonLabel", "Influence Count"))
+				.ToolTipText(LOCTEXT("InfluenceCountSelectionTooltip",
+						"Select vertices that are affected by at least N influences."))
+				.OnClicked_Lambda([this]()
+				{
+					Tool->SelectByInfluenceCount(Tool->GetWeightToolProperties()->ClampSelectValue);
+					return FReply::Handled();
+				})
+			]
+
+			+SHorizontalBox::Slot()
+			.Padding(2.f, WeightEditVerticalPadding)
+			[
+				SNew(SSpinBox<int32>)
+				.MinValue(1)
+				.MaxValue(24)
+				.Value_Lambda([this]()
+				{
+					return ToolSettings->ClampSelectValue;
+				})
+				.OnValueChanged_Lambda([this](int32 NewValue)
+				{
+					ToolSettings->ClampSelectValue = NewValue;
+				})
+				.OnValueCommitted_Lambda([this](int32 NewValue, ETextCommit::Type CommitType)
+				{
+					ToolSettings->SaveConfig();
+				})
+			]
+		]
+	];
+	
+	IDetailCategoryBuilder& EditWeightsCategory = DetailBuilder.EditCategory("Edit Weights", FText::GetEmpty(), ECategoryPriority::Important);
+	EditWeightsCategory.InitiallyCollapsed(true);
+
+	EditWeightsCategory.AddCustomRow(LOCTEXT("SelectMessageRow", "Select Vertices"), false)
+	.WholeRowContent()
+	.VAlign(VAlign_Center)
+	[
+		SNew(STextBlock)
+		.Visibility_Lambda([this]()
+		{
+			return !Tool->HasActiveSelectionOnMainMesh() ? EVisibility::Visible : EVisibility::Collapsed;
+		})
+		.Text_Lambda([this]()
+		{
+			return LOCTEXT("NothingSelectedLabel", "Select vertices on target mesh to edit weights...");
+		})
+	];
+
+	// FLOOD WEIGHTS SLIDER category
+	ToolSettings->DirectEditState.Reset();
+	EditWeightsCategory.AddCustomRow(LOCTEXT("FloodWeightsRow", "Flood Weights Slider"), false)
+	.WholeRowContent()
+	[
+		SNew(SHorizontalBox)
+		.IsEnabled_Lambda([this]{ return Tool->HasActiveSelectionOnMainMesh(); })
+
+		+SHorizontalBox::Slot()
+		[
+			SNew(SSegmentedControl<EWeightEditOperation>)
+			.ToolTipText(LOCTEXT("InteractiveEditModeTooltip",
+				"Add: applies the current weight plus the flood value to the new weight.\n"
+				"Multiply: applies the current weight multiplied by the flood value to the new weight.\n"
+				"This operation applies interactively while dragging the slider. It operates on the currently selected vertices."))
+			.Value_Lambda([this]()
+			{
+				return ToolSettings->DirectEditState.EditMode;
+			})
+			.OnValueChanged_Lambda([this](EWeightEditOperation Mode)
+			{
+				ToolSettings->DirectEditState.EditMode = Mode;
+				ToolSettings->DirectEditState.Reset();
+			})
+			+SSegmentedControl<EWeightEditOperation>::Slot(EWeightEditOperation::Add)
+			.Text(LOCTEXT("BrushAddMode", "Add"))
+			+ SSegmentedControl<EWeightEditOperation>::Slot(EWeightEditOperation::Multiply)
+			.Text(LOCTEXT("BrushMultiplyMode", "Multiply"))
+		]
+		
+		+SHorizontalBox::Slot()
+		.FillWidth(1)
+		[
+			SNew(SVerticalBox)
+
+			+SVerticalBox::Slot()
+			.AutoHeight()
+			[
+				SNew(SSpinBox<float>)
+				.Visibility_Lambda([this]()
+				{
+					const bool bIsVisible = ToolSettings->DirectEditState.EditMode != EWeightEditOperation::Relax;
+					return bIsVisible ? EVisibility::Visible : EVisibility::Collapsed;
+				})
+				.MinSliderValue_Lambda([this]()
+				{
+					return ToolSettings->DirectEditState.GetModeMinValue();
+				})
+				.MaxSliderValue_Lambda([this]()
+				{
+					return ToolSettings->DirectEditState.GetModeMaxValue();
+				})
+				.MinValue_Lambda([this]()
+				{
+					return ToolSettings->DirectEditState.GetModeMinValue();
+				})
+				.MaxValue_Lambda([this]()
+				{
+					return ToolSettings->DirectEditState.GetModeMaxValue();
+				})
+				.Value_Lambda([this]()
+				{
+					return ToolSettings->DirectEditState.CurrentValue;
+				})
+				.OnValueChanged_Lambda([this](float NewValue)
+				{
+					ToolSettings->DirectEditState.CurrentValue = NewValue;
+					
+					if (ToolSettings->DirectEditState.bInTransaction)
+					{
+						
+						float Value = NewValue;
+						if (ToolSettings->DirectEditState.EditMode == EWeightEditOperation::Add)
+						{
+							Value = NewValue - ToolSettings->DirectEditState.StartValue;
+						}
+						
+						constexpr bool bShouldTransact = false;
+						
+						Tool->EditWeightsOnVertices(
+							Tool->GetCurrentBoneIndex(),
+							Value,
+							0, // iterations
+							ToolSettings->DirectEditState.EditMode,
+							Tool->GetMainMeshSelector()->GetSelectedVertices(),
+							bShouldTransact);
+					}
+				})
+				.OnValueCommitted_Lambda([this](float NewValue, ETextCommit::Type CommitType)
+				{
+					if (!ToolSettings->DirectEditState.bInTransaction)
+					{
+						constexpr bool bShouldTransact = true;
+						Tool->EditWeightsOnVertices(
+							Tool->GetCurrentBoneIndex(),
+							NewValue,
+							0, // iterations
+							ToolSettings->DirectEditState.EditMode,
+							Tool->GetMainMeshSelector()->GetSelectedVertices(),
+							bShouldTransact);
+					}
+					ToolSettings->DirectEditState.bInTransaction = false;
+				})
+				.OnBeginSliderMovement_Lambda([this]()
+				{
+					ToolSettings->DirectEditState.StartValue = ToolSettings->DirectEditState.CurrentValue;
+					ToolSettings->DirectEditState.bInTransaction = true;
+					Tool->BeginChange();
+				})
+				.OnEndSliderMovement_Lambda([this](float)
+				{
+					const FText TransactionLabel = LOCTEXT("FloodWeightChange", "Flood weights on vertices.");
+					Tool->EndChange(TransactionLabel);
+					ToolSettings->DirectEditState.bInTransaction = false;
+
+					// reset multiply slider
+					if (ToolSettings->DirectEditState.EditMode == EWeightEditOperation::Multiply)
+					{
+						// multiplying operation is always relative to 1.0
+						ToolSettings->DirectEditState.CurrentValue = 1.0f;
+						ToolSettings->DirectEditState.StartValue = 1.0f;
+					}
+
+					// reset add slider
+					if (ToolSettings->DirectEditState.EditMode == EWeightEditOperation::Add)
+					{
+						// add operation is always relative to 0.0
+						ToolSettings->DirectEditState.CurrentValue = 0.0f;
+						ToolSettings->DirectEditState.StartValue = 0.0f;
+					}
+				})
+				.ToolTipText(LOCTEXT("FloodWeightsToolTip", "Drag the slider to interactively adjust weights on the selected vertices."))
+			]	
+		]
+	];
+
+	// ADD WEIGHTS category
+	EditWeightsCategory.AddCustomRow(LOCTEXT("AddWeightsRow", "Add"), false)
+	.WholeRowContent()
+	[
+		SNew(SVerticalBox)
+		.IsEnabled_Lambda([this]{ return Tool->HasActiveSelectionOnMainMesh(); })
+
+		+ SVerticalBox::Slot()
+		.Padding(WeightEditHorizontalPadding, WeightEditVerticalPadding)
+		[
+			SNew(SHorizontalBox)
+			+SHorizontalBox::Slot()
+			[
+				SNew(SBox)
+				[
+					SNew(SButton)
+					.HAlign(HAlign_Center)
+					.Text(LOCTEXT("AddWeightsButtonLabel", "Add"))
+					.ToolTipText(LOCTEXT("AddButtonTooltip",
+					"Add: applies the current weight plus the flood value to the new weight.\n"
+					"This operation applies to the currently selected vertices."))
+					.OnClicked_Lambda([this]()
+					{
+						constexpr bool bShouldTransact = true;
+						Tool->EditWeightsOnVertices(
+							Tool->GetCurrentBoneIndex(),
+							ToolSettings->AddStrength,
+							0, // iterations 
+							EWeightEditOperation::Add,
+							Tool->GetMainMeshSelector()->GetSelectedVertices(),
+							bShouldTransact);
+						return FReply::Handled();
+					})
+				]
+			]
+
+			+SHorizontalBox::Slot()
+			[
+				SNew(SSpinBox<float>)
+				.MinValue(0)
+				.MaxValue(1)
+				.ToolTipText(LOCTEXT("AddWeightsSliderToolTip", "Adjust the value to Add to the selected vertices."))
+				.Value_Lambda([this]()
+				{
+					return ToolSettings->AddStrength;
+				})
+				.OnValueChanged_Lambda([this](float NewValue)
+				{
+					ToolSettings->AddStrength = NewValue;
+				})
+				.OnValueCommitted_Lambda([this](float NewValue, ETextCommit::Type CommitType)
+				{
+					ToolSettings->SaveConfig();
+				})
+			]
+		]
+	];
+
+	// REPLACE WEIGHTS category
+	EditWeightsCategory.AddCustomRow(LOCTEXT("ReplaceWeightsRow", "Replace"), false)
+	.WholeRowContent()
+	[
+		SNew(SVerticalBox)
+		.IsEnabled_Lambda([this]{ return Tool->HasActiveSelectionOnMainMesh(); })
+
+		+ SVerticalBox::Slot()
+		.Padding(WeightEditHorizontalPadding, WeightEditVerticalPadding)
+		[
+			SNew(SHorizontalBox)
+
+			+SHorizontalBox::Slot()
+			[
+				SNew(SBox)
+				[
+					SNew(SButton)
+					.HAlign(HAlign_Center)
+					.Text(LOCTEXT("ReplaceWeightsButtonLabel", "Replace"))
+					.ToolTipText(LOCTEXT("ReplaceButtonTooltip",
+						"Replace: the weight of selected vertices is replaced by the specified value.\n"
+						"This operation applies to the currently selected vertices."))
+					.OnClicked_Lambda([this]()
+					{
+						constexpr bool bShouldTransact = true;
+						Tool->EditWeightsOnVertices(
+							Tool->GetCurrentBoneIndex(),
+							ToolSettings->ReplaceValue,
+							0, // iterations
+							EWeightEditOperation::Replace,
+							Tool->GetMainMeshSelector()->GetSelectedVertices(),
+							bShouldTransact);
+						return FReply::Handled();
+					})
+				]
+			]
+
+			+SHorizontalBox::Slot()
+			[
+				SNew(SSpinBox<float>)
+				.MinValue(0.f)
+				.MaxValue(1.f)
+				.ToolTipText(LOCTEXT("ReplaceWeightsSliderToolTip", "Adjust the value to Replace on the selected vertices."))
+				.Value_Lambda([this]()
+				{
+					return ToolSettings->ReplaceValue;
+				})
+				.OnValueChanged_Lambda([this](float NewValue)
+				{
+					ToolSettings->ReplaceValue = NewValue;
+				})
+				.OnValueCommitted_Lambda([this](float NewValue, ETextCommit::Type CommitType)
+				{
+					ToolSettings->SaveConfig();
+				})
+			]
+		]
+	];
+
+	// AVERAGE WEIGHTS category
+	EditWeightsCategory.AddCustomRow(LOCTEXT("AverageWeightsRow", "Average"), false)
+	.WholeRowContent()
+	[
+		SNew(SVerticalBox)
+		.IsEnabled_Lambda([this]{ return Tool->HasActiveSelectionOnMainMesh(); })
+
+		+ SVerticalBox::Slot()
+		.Padding(WeightEditHorizontalPadding, WeightEditVerticalPadding)
+		[
+			SNew(SHorizontalBox)
+
+			+SHorizontalBox::Slot()
+			[
+				SNew(SBox)
+				[
+					SNew(SButton)
+					.HAlign(HAlign_Center)
+					.Text(LOCTEXT("AverageWeightsButtonLabel", "Average"))
+					.ToolTipText(LOCTEXT("AverageButtonTooltip",
+						"Weights on all selected vertices are set to the average of all selected vertices."))
+					.OnClicked_Lambda([this]()
+					{
+						ToolSettings->WeightTool->AverageWeights(ToolSettings->AverageStrength);
+						return FReply::Handled();
+					})
+				]
+			]
+
+			+SHorizontalBox::Slot()
+			[
+				SNew(SSpinBox<float>)
+				.MinValue(0.f)
+				.MaxValue(1.f)
+				.ToolTipText(LOCTEXT("AverageWeightsSliderToolTip", "Blend the amount to Average the weights on the selected vertices."))
+				.Value_Lambda([this]()
+				{
+					return ToolSettings->AverageStrength;
+				})
+				.OnValueChanged_Lambda([this](float NewValue)
+				{
+					ToolSettings->AverageStrength = NewValue;
+				})
+				.OnValueCommitted_Lambda([this](float NewValue, ETextCommit::Type CommitType)
+				{
+					ToolSettings->SaveConfig();
+				})
+			]
+		]
+	];
+
+	// RELAX WEIGHTS category
+	EditWeightsCategory.AddCustomRow(LOCTEXT("RelaxWeightsRow", "Relax"), false)
+	.WholeRowContent()
+	[
+		SNew(SVerticalBox)
+		.IsEnabled_Lambda([this]{ return Tool->HasActiveSelectionOnMainMesh(); })
+
+		+ SVerticalBox::Slot()
+		.Padding(WeightEditHorizontalPadding, WeightEditVerticalPadding)
+		[
+			SNew(SHorizontalBox)
+			+SHorizontalBox::Slot()
+			[
+				SNew(SBox)
+				[
+					SNew(SButton)
+					.HAlign(HAlign_Center)
+					.Text(LOCTEXT("RelaxWeightsButtonLabel", "Relax"))
+					.ToolTipText(LOCTEXT("RelaxButtonTooltip",
+						"Relax: the weight of each selected vertex is replaced by the average of it's neighbors. \n"
+						"This smooths weights across the mesh."))
+					.OnClicked_Lambda([this]()
+					{
+						constexpr bool bShouldTransact = true;
+						constexpr int32 DefaultRelaxIterations = 5; // provides a reasonable falloff distance 
+						Tool->EditWeightsOnVertices(
+							Tool->GetCurrentBoneIndex(),
+							ToolSettings->RelaxStrength,
+							DefaultRelaxIterations, 
+							EWeightEditOperation::Relax,
+							Tool->GetMainMeshSelector()->GetSelectedVertices(),
+							bShouldTransact);
+						return FReply::Handled();
+					})
+				]
+			]
+
+			+SHorizontalBox::Slot()
+			[
+				SNew(SSpinBox<float>)
+				.MinValue(0)
+				.MaxValue(1)
+				.ToolTipText(LOCTEXT("RelaxWeightsSliderToolTip", "Blend the amount to Relax the weights on the selected vertices."))
+				.Value_Lambda([this]()
+				{
+					return ToolSettings->RelaxStrength;
+				})
+				.OnValueChanged_Lambda([this](float NewValue)
+				{
+					ToolSettings->RelaxStrength = NewValue;
+				})
+				.OnValueCommitted_Lambda([this](float NewValue, ETextCommit::Type CommitType)
+				{
+					ToolSettings->SaveConfig();
+				})
+			]
+		]
+	];
+	
+	// MIRROR WEIGHTS category
+	EditWeightsCategory.AddCustomRow(LOCTEXT("MirrorWeightsRow", "Mirror"), false)
+	.WholeRowContent()
+	[
+		SNew(SVerticalBox)
+		.IsEnabled_Lambda([this]{ return Tool->HasActiveSelectionOnMainMesh(); })
+		
+		+ SVerticalBox::Slot()
+		.Padding(WeightEditHorizontalPadding, WeightEditVerticalPadding)
+		[
+			SNew(SHorizontalBox)
+			
+			+SHorizontalBox::Slot()
+			.VAlign(VAlign_Center)
+			.FillWidth(WeightEditingLabelsPercent)
+			[
+				SNew(STextBlock)
+				.Text(LOCTEXT("MirrorPlaneLabel", "Mirror Plane"))
+				.Font(FAppStyle::GetFontStyle(TEXT("PropertyWindow.NormalFont")))
+				.ToolTipText(LOCTEXT("MirrorPlaneTooltip", "The plane to copy weights across."))
+			]
+
+			+SHorizontalBox::Slot()
+			.FillWidth(1.f)
+			[
+				SNew(SHorizontalBox)
+				+SHorizontalBox::Slot()
+				[
+					SNew(SSegmentedControl<EAxis::Type>)
+					.ToolTipText(LOCTEXT("MirrorAxisTooltip",
+						"X: copies weights across the YZ plane.\n"
+						"Y: copies weights across the XZ plane.\n"
+						"Z: copies weights across the XY plane."))
+					.Value_Lambda([this]()
+					{
+						return ToolSettings->MirrorAxis;
+					})
+					.OnValueChanged_Lambda([this](EAxis::Type Mode)
+					{
+						ToolSettings->MirrorAxis = Mode;
+					})
+					+ SSegmentedControl<EAxis::Type>::Slot(EAxis::X)
+					.Text(LOCTEXT("MirrorXLabel", "X"))
+					+ SSegmentedControl<EAxis::Type>::Slot(EAxis::Y)
+					.Text(LOCTEXT("MirrorYLabel", "Y"))
+					+ SSegmentedControl<EAxis::Type>::Slot(EAxis::Z)
+					.Text(LOCTEXT("MirrorZLabel", "Z"))
+				]
+					
+				+SHorizontalBox::Slot()
+				[
+					SNew(SSegmentedControl<EMirrorDirection>)
+					.ToolTipText(LOCTEXT("MirrorDirectionTooltip", "The direction that determines what side of the plane to copy weights from."))
+					.Value_Lambda([this]()
+					{
+						return ToolSettings->MirrorDirection;
+					})
+					.OnValueChanged_Lambda([this](EMirrorDirection Mode)
+					{
+						ToolSettings->MirrorDirection = Mode;
+					})
+					+ SSegmentedControl<EMirrorDirection>::Slot(EMirrorDirection::PositiveToNegative)
+					.Text(LOCTEXT("MirrorPosToNegLabel", "+ to -"))
+					+ SSegmentedControl<EMirrorDirection>::Slot(EMirrorDirection::NegativeToPositive)
+					.Text(LOCTEXT("MirrorNegToPosLabel", "- to +"))
+				]
+			]
+		]
+
+		+SVerticalBox::Slot()
+		.Padding(0.f, WeightEditVerticalPadding)
+		[
+			SNew(SBox)
+			[
+				SNew(SButton)
+				.HAlign(HAlign_Center)
+				.Text(LOCTEXT("MirrorWeightsButtonLabel", "Mirror"))
+				.ToolTipText(LOCTEXT("MirrorButtonTooltip",
+					"Weights are copied across the given plane in the given direction.\n"
+					"This command operates on the selected vertices."))
+				.OnClicked_Lambda([this]()
+				{
+					Tool->MirrorWeights(ToolSettings->MirrorAxis, ToolSettings->MirrorDirection);
+					return FReply::Handled();
+				})
+			]
+		]
+	];
+
+	// NORMALIZE WEIGHTS category
+	EditWeightsCategory.AddCustomRow(LOCTEXT("NormalizeWeightsRow", "Normalize"), false)
+	.WholeRowContent()
+	[
+		SNew(SHorizontalBox)
+		.IsEnabled_Lambda([this]{ return Tool->HasActiveSelectionOnMainMesh(); })
+
+		+SHorizontalBox::Slot()
+		.Padding(2.f, WeightEditVerticalPadding)
+		[
+			SNew(SButton)
+			.HAlign(HAlign_Center)
+			.VAlign(VAlign_Center)
+			.Text(LOCTEXT("NormalizeWeightsButtonLabel", "Normalize"))
+			.ToolTipText(LOCTEXT("NormalizeWeightsTooltip",
+					"Forces the weights on the selected vertices to sum to 1.\n"
+					"This command operates on the selected vertices."))
+			.IsEnabled_Lambda([this]()
+			{
+				return ToolSettings->EditingMode == EWeightEditMode::Mesh;
+			})
+			.OnClicked_Lambda([this]()
+			{
+				Tool->NormalizeWeights();
+				return FReply::Handled();
+			})
+		]
+
+		+SHorizontalBox::Slot()
+		.Padding(2.f, WeightEditVerticalPadding)
+		[
+			SNew(SButton)
+			.HAlign(HAlign_Center)
+			.VAlign(VAlign_Center)
+			.Text(LOCTEXT("HammerWeightsButtonLabel", "Hammer"))
+			.ToolTipText(LOCTEXT("HammerWeightsTooltip",
+					"Copies the weight of the nearest non-selected vertex.\n"
+					"This command operates on the selected vertices."))
+			.IsEnabled_Lambda([this]()
+			{
+				return ToolSettings->EditingMode == EWeightEditMode::Mesh;
+			})
+			.OnClicked_Lambda([this]()
+			{
+				Tool->HammerWeights();
+				return FReply::Handled();
+			})
+		]
+	];
+
+	// PRUNE WEIGHTS category
+	EditWeightsCategory.AddCustomRow(LOCTEXT("PruneWeightsRow", "Prune"), false)
+	.WholeRowContent()
+	[
+		SNew(SHorizontalBox)
+		.IsEnabled_Lambda([this]{ return Tool->HasActiveSelectionOnMainMesh(); })
+
+		+SHorizontalBox::Slot()
+		.Padding(2.f, WeightEditVerticalPadding)
+		[
+			SNew(SButton)
+			.HAlign(HAlign_Center)
+			.Text(LOCTEXT("PruneWeightsButtonLabel", "Prune"))
+			.ToolTipText(LOCTEXT("PruneButtonTooltip",
+				"Removes influences with weights below the given threshold value.\n"
+				"Pruned bones are removed from the list of bones affecting the given vertex.\n"
+				"Pruned bones will no longer recieve weight when a vertex is normalized.\n"
+				"This command operates on the selected vertices."))
+			.OnClicked_Lambda([this]()
+			{
+				ToolSettings->WeightTool->PruneWeights(ToolSettings->PruneValue, TArray<BoneIndex>());
+				return FReply::Handled();
+			})
+		]
+
+		+SHorizontalBox::Slot()
+		.Padding(2.f, WeightEditVerticalPadding)
+		[
+			SNew(SSpinBox<float>)
+			.MinValue(0.f)
+			.MaxValue(1.f)
+			.Value_Lambda([this]()
+			{
+				return ToolSettings->PruneValue;
+			})
+			.OnValueChanged_Lambda([this](float NewValue)
+			{
+				ToolSettings->PruneValue = NewValue;
+			})
+			.OnValueCommitted_Lambda([this](float NewValue, ETextCommit::Type CommitType)
+			{
+				ToolSettings->SaveConfig();
+			})
+		]
+	];
+
+	// CLAMP INFLUENCES category
+	EditWeightsCategory.AddCustomRow(LOCTEXT("ClampInfluencesRow", "Clamp"), false)
+	.WholeRowContent()
+	[
+		SNew(SHorizontalBox)
+		.IsEnabled_Lambda([this]{ return Tool->HasActiveSelectionOnMainMesh(); })
+
+		+SHorizontalBox::Slot()
+		.Padding(2.f, WeightEditVerticalPadding)
+		[
+			SNew(SButton)
+			.HAlign(HAlign_Center)
+			.Text(LOCTEXT("ClampInfluencesButtonLabel", "Clamp"))
+			.ToolTipText(LOCTEXT("ClampInfluencesButtonTooltip",
+				"Clamp the number of influences to not exceed the target value.\n"
+				"Removes smallest influences first.\n"
+				"This command operates on the selected vertices."))
+			.OnClicked_Lambda([this]()
+			{
+				ToolSettings->WeightTool->ClampInfluences(ToolSettings->ClampValue);
+				return FReply::Handled();
+			})
+		]
+
+		+SHorizontalBox::Slot()
+		.Padding(2.f, WeightEditVerticalPadding)
+		[
+			SNew(SSpinBox<int32>)
+			.MinValue(1)
+			.MaxValue(24)
+			.Value_Lambda([this]()
+			{
+				return ToolSettings->ClampValue;
+			})
+			.OnValueChanged_Lambda([this](int32 NewValue)
+			{
+				ToolSettings->ClampValue = NewValue;
+			})
+			.OnValueCommitted_Lambda([this](int32 NewValue, ETextCommit::Type CommitType)
+			{
+				ToolSettings->SaveConfig();
+			})
+		]
+	];
+
+	// COPY/PASTE WEIGHTS category
+	EditWeightsCategory.AddCustomRow(LOCTEXT("CopyPasteWeightsRow", "Copy Paste"), false)
+	.WholeRowContent()
+	[
+		SNew(SHorizontalBox)
+		.IsEnabled_Lambda([this]{ return Tool->HasActiveSelectionOnMainMesh(); })
+
+		+SHorizontalBox::Slot()
+		.Padding(2.f, WeightEditVerticalPadding)
+		[
+			SNew(SBox)
+			[
+				SNew(SButton)
+				.HAlign(HAlign_Center)
+				.Text(LOCTEXT("CopyWeightsButtonLabel", "Copy"))
+				.ToolTipText(LOCTEXT("CopyButtonTooltip",
+					"Copy the average weights of the selected vertices to the clipboard. \n"
+					"This is designed to work with the Paste command."))
+				.OnClicked_Lambda([this]()
+				{
+					ToolSettings->WeightTool->CopyWeights();
+					return FReply::Handled();
+				})
+			]
+		]
+
+		+SHorizontalBox::Slot()
+		.Padding(2.f, WeightEditVerticalPadding)
+		[
+			SNew(SButton)
+			.HAlign(HAlign_Center)
+			.Text(LOCTEXT("PasteWeightsButtonLabel", "Paste"))
+			.ToolTipText(LOCTEXT("PasteButtonTooltip",
+				"Paste the weights on the selected vertices.\n"
+				"This command requires the clipboard contain weights from the Copy command."))
+			.OnClicked_Lambda([this]()
+			{
+				ToolSettings->WeightTool->PasteWeights();
+				return FReply::Handled();
+			})
+		]
+	];
+	
+	// VERTEX EDITOR category
+	EditWeightsCategory.AddCustomRow(LOCTEXT("VertexEditorRow", "Component Editor"), false)
+	.WholeRowContent()
+	[
+		SNew(SBox)
+		.IsEnabled_Lambda([this]{ return Tool->HasActiveSelectionOnMainMesh(); })
+		[
+			SNew(SVertexWeightEditor, ToolSettings->WeightTool)
+		]
+	];
+}
+
+void FSkinWeightDetailCustomization::AddTransferUI(IDetailLayoutBuilder& DetailBuilder) const
+{
+	if (!ensure(Tool.IsValid()))
+	{
+		return;
+	}
+	
+	IDetailCategoryBuilder& TransferWeightsCategory = DetailBuilder.EditCategory("WeightTransfer", FText::GetEmpty(), ECategoryPriority::Important);
+	TransferWeightsCategory.InitiallyCollapsed(true);
+
+	// TRANSFER BUTTON
+	TransferWeightsCategory.AddCustomRow(LOCTEXT("TransferWeightsRow", "Transfer Weights"), false)
+	.WholeRowContent()
+	[
+		SNew(SBox)
+		.Padding(0.f, WeightEditVerticalPadding)
+		[
+			SNew(SButton)
+			.HAlign(HAlign_Center)
+			.Text(LOCTEXT("TransferWeightsButtonLabel", "Transfer Weights"))
+			.ToolTipText(LOCTEXT("TransferButtonTooltip",
+						"Weights are transferred from the source skeletal mesh.\n"
+						"Vertices may be selected on the source and/or target mesh to filter which parts to copy from and which parts to copy to.\n"
+						"If either mesh has no vertices selected, the whole mesh is considered.\n"))
+			.OnClicked_Lambda([this]()
+			{
+				Tool->GetWeightTransferManager()->TransferWeights();
+				return FReply::Handled();
+			})
+			.IsEnabled_Lambda([this]()
+			{
+				return Tool->GetWeightTransferManager()->CanTransferWeights();
+			})
+		]
+	];
+
+	// SKELETAL MESH ASSET INPUT
+	const TSharedRef<IPropertyHandle> SourceSkeletalMeshHandle = DetailBuilder.GetProperty(GET_MEMBER_NAME_CHECKED(USkinWeightsPaintToolProperties, SourceSkeletalMesh), USkinWeightsPaintToolProperties::StaticClass());
+	TransferWeightsCategory.AddProperty(SourceSkeletalMeshHandle);
+	
+	// LOD
+	const TSharedRef<IPropertyHandle> LODHandle = DetailBuilder.GetProperty(GET_MEMBER_NAME_CHECKED(USkinWeightsPaintToolProperties, SourceLOD), USkinWeightsPaintToolProperties::StaticClass());
+	TransferWeightsCategory.AddProperty(LODHandle);
+
+	// PROFILE
+	const TSharedRef<IPropertyHandle> SourceProfileHandle = DetailBuilder.GetProperty(GET_MEMBER_NAME_CHECKED(USkinWeightsPaintToolProperties, SourceSkinWeightProfile), USkinWeightsPaintToolProperties::StaticClass());
+	TransferWeightsCategory.AddProperty(SourceProfileHandle);
+
+	// MESH SELECTION OPTION (SOURCE OR TARGET)
+	const TSharedRef<IPropertyHandle> TransferSelectModeHandle = DetailBuilder.GetProperty(GET_MEMBER_NAME_CHECKED(USkinWeightsPaintToolProperties, MeshSelectMode), USkinWeightsPaintToolProperties::StaticClass());
+	FDetailWidgetRow& TransferSelectModeRow = TransferWeightsCategory.AddCustomRow(LOCTEXT("SelectionModeRow", "Selection Mode"), false)
+	.NameContent()
+	[
+		TransferSelectModeHandle->CreatePropertyNameWidget()
+	]
+	.ValueContent()
+	[
+		SNew(SBox)
+		.Padding(2.0f)
+		.HAlign(HAlign_Center)
+		[
+			SNew(SSegmentedControl<EMeshTransferOption>)
+			.ToolTipText(LOCTEXT("SelectionSourceTooltip",
+					"Choose which mesh to select components on (vertices/edges/faces).\n"
+					"Weights will be transferred from selected components on the source to selected components on the target.\n"
+					"If no components are selected on either the source or target, the whole mesh will be considered.\n"
+					"Source: The mesh to copy weights FROM.\n"
+					"Target: The mesh to copy weights TO (the main mesh in the tool)."))
+			.Value_Lambda([this]()
+			{
+				return ToolSettings->MeshSelectMode;
+			})
+			.OnValueChanged_Lambda([this](EMeshTransferOption Mode)
+			{
+				ToolSettings->MeshSelectMode = Mode;
+				ToolSettings->WeightTool->UpdateSelectorState();
+			})
+			+ SSegmentedControl<EMeshTransferOption>::Slot(EMeshTransferOption::Source)
+			.Text(LOCTEXT("SourceMode", "Source"))
+			+SSegmentedControl<EMeshTransferOption>::Slot(EMeshTransferOption::Target)
+			.Text(LOCTEXT("TargetMode", "Target"))
+		]
+	];
+
+	// PREVIEW OFFSET
+	const TSharedRef<IPropertyHandle> PreviewOffsetHandle = DetailBuilder.GetProperty(GET_MEMBER_NAME_CHECKED(USkinWeightsPaintToolProperties, SourcePreviewOffset), USkinWeightsPaintToolProperties::StaticClass());
+	TransferWeightsCategory.AddProperty(PreviewOffsetHandle);
+}
+
+void SVertexWeightItem::Construct(const FArguments& InArgs, const TSharedRef<STableViewBase>& OwnerTableView)
+{
+	Element = InArgs._Element;
+	ParentTable = InArgs._ParentTable;
+	SMultiColumnTableRow<TSharedPtr<FWeightEditorElement>>::Construct(FSuperRowType::FArguments(), OwnerTableView);
+}
+
+TSharedRef<SWidget> SVertexWeightItem::GenerateWidgetForColumn(const FName& ColumnName)
+{
+	if (ColumnName == ColumnName_Bone)
+	{
+		const FName BoneName = ParentTable->Tool->GetBoneNameFromIndex(Element->BoneIndex);
+		return SNew(STextBlock).Text(FText::FromName(BoneName));
+	}
+
+	if (ColumnName == ColumnName_Weight)
+	{
+		return SNew(SNumericEntryBox<float>)
+		.AllowSpin(true)
+		.MinSliderValue(0.f)
+		.MinValue(0.f)
+		.MaxSliderValue(1.0f)
+		.MaxValue(1.f)
+		.Value_Lambda([this]()
+		{
+			if (bInTransaction)
+			{
+				return ValueDuringSlide;
+			}
+			
+			return ParentTable->Tool->GetAverageWeightOnBone(Element->BoneIndex, ParentTable->Tool->GetMainMeshSelector()->GetSelectedVertices());
+		})
+		.OnValueChanged_Lambda([this](float NewValue)
+		{
+			if (bInTransaction)
+			{
+				ValueDuringSlide = NewValue;
+
+				const bool bScalingUp = NewValue >= ValueAtStartOfSlide || FMath::IsNearlyEqual(ValueAtStartOfSlide, 0.f);
+				const float RangeEnd = bScalingUp ? 1.f : 0.f;
+				float RelativeScale = (NewValue - ValueAtStartOfSlide) * 1.f / (RangeEnd - ValueAtStartOfSlide);
+				RelativeScale *= bScalingUp ? 1.f : -1.f;
+				constexpr bool bShouldTransact = false;
+				ParentTable->Tool->EditWeightsOnVertices(
+					Element->BoneIndex,
+					RelativeScale,
+					0, // iterations
+					EWeightEditOperation::RelativeScale,
+					ParentTable->Tool->GetMainMeshSelector()->GetSelectedVertices(),
+					bShouldTransact);
+			}
+		})
+		.OnValueCommitted_Lambda([this](float NewValue, ETextCommit::Type CommitType)
+		{
+			if (!bInTransaction)
+			{
+				constexpr bool bShouldTransact = true;
+				ParentTable->Tool->EditWeightsOnVertices(
+					Element->BoneIndex,
+					NewValue,
+					0, // iterations,
+					EWeightEditOperation::Replace,
+					ParentTable->Tool->GetMainMeshSelector()->GetSelectedVertices(),
+					bShouldTransact);
+			}
+			bInTransaction = false;
+		})
+		.OnBeginSliderMovement_Lambda([this]()
+		{
+			ParentTable->Tool->BeginChange();
+			ValueAtStartOfSlide = ParentTable->Tool->GetAverageWeightOnBone(Element->BoneIndex, ParentTable->Tool->GetMainMeshSelector()->GetSelectedVertices());
+			ValueDuringSlide = ValueAtStartOfSlide;
+			bInTransaction = true;
+		})
+		.OnEndSliderMovement_Lambda([this](float)
+		{
+			const FText TransactionLabel = LOCTEXT("DirectWeightChange", "Scale weights on vertices.");
+			ParentTable->Tool->EndChange(TransactionLabel);
+			bInTransaction = false;
+		})
+		.ToolTipText(LOCTEXT("WeightSliderToolTip", "Set the weight on this bone for the selected vertices."));
+	}
+
+	if (ColumnName == ColumnName_Prune)
+	{
+		return SNew(SHorizontalBox)
+		
+		+SHorizontalBox::Slot()
+		.FillWidth(1.f)
+		.HAlign(HAlign_Center)
+		.VAlign(VAlign_Center)
+		[
+			SNew(SButton)
+			.HAlign(HAlign_Right)
+			.ToolTipText(LOCTEXT("PruneInfluence", "Prune the influence from the selected vertices."))
+			.OnClicked_Lambda([this]() -> FReply
+			{
+				// use a negative threshold weight that no weight will ever
+				// be below because we only want to prune based on the bone, regardless of the weight
+				constexpr float NegativeThreshold = -1.0f;
+				ParentTable->Tool->PruneWeights(NegativeThreshold, {Element->BoneIndex});
+				return FReply::Handled();
+			})
+			.Content()
+			[
+				SNew(SImage)
+				.Image(FAppStyle::Get().GetBrush("Icons.Delete"))
+				.ColorAndOpacity(FSlateColor::UseForeground())
+			]
+		];
+	}
+
+	checkNoEntry();
+	return SNullWidget::NullWidget;
+}
+
+SVertexWeightEditor::~SVertexWeightEditor()
+{
+	if (Tool.IsValid())
+	{
+		Tool->OnSelectionChanged.RemoveAll(this);
+		Tool->OnWeightsChanged.RemoveAll(this);
+		Tool.Reset();
+	}
+}
+
+void SVertexWeightEditor::Construct(const FArguments& InArgs, USkinWeightsPaintTool* InSkinTool)
+{
+	Tool = InSkinTool;
+
+	ChildSlot
+	[
+		SNew(SBox)
+		[
+			SAssignNew( ListView, SWeightEditorListViewType )
+			.Visibility_Lambda([this]()
+			{
+				return Tool->GetMainMeshSelector()->GetSelectedVertices().IsEmpty() ? EVisibility::Collapsed : EVisibility::Visible;
+			})
+			.SelectionMode(ESelectionMode::Single)
+			.ListItemsSource( &ListViewItems )
+			.OnGenerateRow_Lambda([this](TSharedPtr<FWeightEditorElement> Element, const TSharedRef<STableViewBase>& OwnerTableView)
+			{
+				return SNew(SVertexWeightItem, OwnerTableView).Element(Element).ParentTable(SharedThis(this));
+			})
+			.HeaderRow
+			(
+				SNew(SHeaderRow)
+				+ SHeaderRow::Column(ColumnName_Bone)
+				.HAlignHeader(HAlign_Center)
+				.DefaultLabel(LOCTEXT("WeightEditorBoneColumn", "Bone"))
+				+ SHeaderRow::Column(ColumnName_Weight)
+				.HAlignHeader(HAlign_Center)
+				.DefaultLabel(LOCTEXT("WeightEditorWeightColumn", "Weight (Average)"))
+				+ SHeaderRow::Column(ColumnName_Prune)
+				.HAlignHeader(HAlign_Center)
+				.DefaultLabel(LOCTEXT("WeightEditorPruneColumn", "Prune"))
+				.FixedWidth(60.f)
+			)
+		]
+	];
+
+	RefreshView();
+	
+	Tool->OnSelectionChanged.AddSP(this, &SVertexWeightEditor::RefreshView);
+	Tool->OnWeightsChanged.AddSP(this, &SVertexWeightEditor::RefreshView);
+}
+
+void SVertexWeightEditor::RefreshView()
+{
+	if (!Tool.IsValid())
+	{
+		return; 
+	}
+	
+	// get all bones affecting the selected vertices
+	TArray<int32> Influences;
+	Tool->GetInfluences(Tool->GetMainMeshSelector()->GetSelectedVertices(), Influences);
+
+	// generate list view items
+	ListViewItems.Reset();
+	for (const int32 InfluenceIndex : Influences)
+	{
+		ListViewItems.Add(MakeShareable(new FWeightEditorElement(InfluenceIndex)));
+	}
+	
+	ListView->RequestListRefresh();
+}
+
+#undef LOCTEXT_NAMESPACE
diff --git a/Engine/Plugins/Animation/SkeletalMeshModelingTools/Source/SkeletalMeshModelingTools/Private/DetailCustomization/SkinWeightDetailCustomization.h b/Engine/Plugins/Animation/SkeletalMeshModelingTools/Source/SkeletalMeshModelingTools/Private/DetailCustomization/SkinWeightDetailCustomization.h
new file mode 100644
--- /dev/null
+++ b/Engine/Plugins/Animation/SkeletalMeshModelingTools/Source/SkeletalMeshModelingTools/Private/DetailCustomization/SkinWeightDetailCustomization.h
@@ -0,0 +1,103 @@
+﻿// Copyright Epic Games, Inc. All Rights Reserved.
+
+#pragma once
+
+#include "CoreMinimal.h"
+#include "IDetailCustomization.h"
+#include "SkeletalMesh/SkinWeightsPaintTool.h"
+#include "UObject/WeakObjectPtr.h"
+#include "Widgets/Views/STableRow.h"
+#include "Widgets/Views/SListView.h"
+
+class USkinWeightsPaintTool;
+class USkinWeightsPaintToolProperties;
+
+class FSkinWeightDetailCustomization : public IDetailCustomization
+{
+public:
+
+	virtual ~FSkinWeightDetailCustomization() override;
+
+	static TSharedRef<IDetailCustomization> MakeInstance()
+	{
+		return MakeShareable(new FSkinWeightDetailCustomization);
+	}
+
+	/** IDetailCustomization interface */
+	virtual void CustomizeDetails(IDetailLayoutBuilder& DetailBuilder) override;
+
+	void OnSelectionChanged()
+	{
+		ToolSettings.Get()->DirectEditState.Reset();
+	};
+
+private:
+
+	void AddBrushUI(IDetailLayoutBuilder& DetailBuilder) const;
+	void AddSelectionUI(IDetailLayoutBuilder& DetailBuilder) const;
+	void AddTransferUI(IDetailLayoutBuilder& DetailBuilder) const;
+
+	IDetailLayoutBuilder* CurrentDetailBuilder = nullptr;
+	TWeakObjectPtr<USkinWeightsPaintToolProperties> ToolSettings;
+	TWeakObjectPtr<USkinWeightsPaintTool> Tool;
+
+	static float WeightSliderWidths;
+	static float WeightEditingLabelsPercent;
+	static float WeightEditVerticalPadding;
+	static float WeightEditHorizontalPadding;
+};
+
+struct FWeightEditorElement
+{
+	int32 BoneIndex;
+
+	FWeightEditorElement(int32 InBoneIndex) : BoneIndex(InBoneIndex) {}
+};
+
+class SVertexWeightEditor;
+class SVertexWeightItem : public SMultiColumnTableRow<TSharedPtr<FWeightEditorElement>>
+{
+public:
+
+	SLATE_BEGIN_ARGS(SVertexWeightItem) {}
+	SLATE_ARGUMENT(TSharedPtr<FWeightEditorElement>, Element)
+	SLATE_ARGUMENT(TSharedPtr<SVertexWeightEditor>, ParentTable)
+	SLATE_END_ARGS()
+
+	virtual ~SVertexWeightItem() override {};
+
+	void Construct(const FArguments& InArgs, const TSharedRef<STableViewBase>& OwnerTableView);
+
+	virtual TSharedRef<SWidget> GenerateWidgetForColumn(const FName& ColumnName) override;
+
+private:
+	TSharedPtr<FWeightEditorElement> Element;
+	TSharedPtr<SVertexWeightEditor> ParentTable;
+	bool bInTransaction = false;
+	float ValueAtStartOfSlide;
+	float ValueDuringSlide;
+};
+
+typedef SListView< TSharedPtr<FWeightEditorElement> > SWeightEditorListViewType;
+
+class SVertexWeightEditor : public SCompoundWidget
+{
+public:
+	SLATE_BEGIN_ARGS(SVertexWeightEditor) {}
+	SLATE_END_ARGS()
+
+	void Construct(const FArguments& InArgs, USkinWeightsPaintTool* InSkinTool);
+
+	virtual ~SVertexWeightEditor() override;
+	
+	void RefreshView();
+
+private:
+
+	// the vertex weight list view
+	TSharedPtr<SWeightEditorListViewType> ListView;
+	TArray< TSharedPtr<FWeightEditorElement> > ListViewItems;
+	TWeakObjectPtr<USkinWeightsPaintTool> Tool;
+
+	friend SVertexWeightItem;
+};
\ No newline at end of file
diff --git a/Engine/Plugins/Animation/SkeletalMeshModelingTools/Source/SkeletalMeshModelingTools/Private/SReferenceSkeletonTree.cpp b/Engine/Plugins/Animation/SkeletalMeshModelingTools/Source/SkeletalMeshModelingTools/Private/SReferenceSkeletonTree.cpp
new file mode 100644
--- /dev/null
+++ b/Engine/Plugins/Animation/SkeletalMeshModelingTools/Source/SkeletalMeshModelingTools/Private/SReferenceSkeletonTree.cpp
@@ -0,0 +1,1135 @@
+// Copyright Epic Games, Inc. All Rights Reserved.
+
+#include "SReferenceSkeletonTree.h"
+
+#include "SkeletonModifier.h"
+#include "SkeletalMeshModelingToolsCommands.h"
+#include "SkeletonClipboard.h"
+
+#include "SPositiveActionButton.h"
+
+#include "Misc/TextFilterExpressionEvaluator.h"
+#include "Widgets/Input/SSearchBox.h"
+#include "Widgets/Text/SInlineEditableTextBlock.h"
+
+#include "Framework/MultiBox/MultiBoxBuilder.h"
+#include "Editor.h"
+
+#include "Misc/ITransaction.h"
+
+#define LOCTEXT_NAMESPACE "SReferenceSkeletonTree"
+
+namespace ReferenceSkeletonTreeLocals
+{
+
+FSkeletonModifierChange::FSkeletonModifierChange(const USkeletonModifier* InModifier)
+	: FCommandChange()
+	, PreChangeSkeleton(InModifier->GetReferenceSkeleton())
+	, PreBoneTracker(InModifier->GetBoneIndexTracker())
+	, PostChangeSkeleton(InModifier->GetReferenceSkeleton())
+	, PostBoneTracker(InModifier->GetBoneIndexTracker())
+{}
+
+void FSkeletonModifierChange::StoreSkeleton(const USkeletonModifier* InModifier)
+{
+	PostChangeSkeleton = InModifier->GetReferenceSkeleton();
+	PostBoneTracker = InModifier->GetBoneIndexTracker();
+}
+
+void FSkeletonModifierChange::Apply(UObject* Object)
+{ // redo
+	USkeletonModifier* Modifier = CastChecked<USkeletonModifier>(Object);
+	Modifier->ExternalUpdate(PostChangeSkeleton, PostBoneTracker);
+}
+
+void FSkeletonModifierChange::Revert(UObject* Object)
+{ // undo
+	USkeletonModifier* Modifier = CastChecked<USkeletonModifier>(Object);
+	Modifier->ExternalUpdate(PreChangeSkeleton, PreBoneTracker);
+}
+
+}
+
+FBoneElement::FBoneElement(const FName& InBoneName, TWeakObjectPtr<USkeletonModifier> InModifier)
+	: BoneName(InBoneName)
+	, WeakModifier(InModifier)
+{}
+
+void FBoneElement::RequestRename() const
+{
+	OnRenameRequested.ExecuteIfBound();
+}
+
+void SBoneItem::Construct(const FArguments& InArgs)
+{
+	// cf. FSkeletonTreeBoneItem::GenerateWidgetForNameColumn
+	
+	WeakTreeElement = InArgs._TreeElement;
+	WeakModifier = InArgs._WeakModifier;
+
+	const TSharedPtr<FBoneElement> TreeElement = WeakTreeElement.Pin();
+
+	const FSlateBrush* Brush = FAppStyle::Get().GetBrush("SkeletonTree.Bone");
+
+	TSharedPtr<SHorizontalBox> RowBox;
+
+	ChildSlot
+	[
+		SAssignNew(RowBox, SHorizontalBox)
+		+ SHorizontalBox::Slot()
+		.AutoWidth()
+		.Padding(FMargin(0.0f, 2.0f))
+		.VAlign(VAlign_Center)
+		.HAlign(HAlign_Center)
+		[
+			SNew(SImage)
+			.Image(Brush)
+			.ColorAndOpacity_Static(&SBoneItem::GetTextColor, InArgs._IsSelected)
+		]
+	];
+
+	TSharedPtr< SInlineEditableTextBlock > InlineWidget;
+	
+	RowBox->AddSlot()
+    .AutoWidth()
+	.Padding(4, 0, 0, 0)
+    .VAlign(VAlign_Center)
+    [
+	    SAssignNew(InlineWidget, SInlineEditableTextBlock)
+		.Text(this, &SBoneItem::GetName)
+		// .Font(TextFont)
+		.ColorAndOpacity_Static(&SBoneItem::GetTextColor, InArgs._IsSelected)
+		.MaximumLength(NAME_SIZE-1)
+		.OnVerifyTextChanged(this, &SBoneItem::OnVerifyNameChanged)
+		.OnTextCommitted(InArgs._Delegates.OnBoneNameCommitted)
+		.MultiLine(false)
+    ];
+
+	TreeElement->OnRenameRequested.BindSP(InlineWidget.Get(), &SInlineEditableTextBlock::EnterEditingMode);
+}
+
+FSlateColor SBoneItem::GetTextColor(FIsSelected InIsSelected)
+{
+	const bool bIsSelected = InIsSelected.IsBound() ? InIsSelected.Execute() : false;
+	if (bIsSelected)
+	{
+		return FAppStyle::Get().GetSlateColor("Colors.ForegroundInverted");
+	}
+	
+	return FSlateColor::UseForeground();
+}
+
+bool SBoneItem::OnVerifyNameChanged(const FText& InText, FText& OutErrorMessage) const
+{
+	if (!WeakModifier.IsValid())
+	{
+		return false;
+	}
+
+	const TSharedPtr<FBoneElement> BoneElementPtr = WeakTreeElement.Pin();
+	const FReferenceSkeleton& ReferenceSkeleton = WeakModifier->GetReferenceSkeleton();
+
+	const FName NewName = FName(InText.ToString());
+	const int32 BoneIndex = ReferenceSkeleton.FindRawBoneIndex(NewName);
+	return BoneIndex == INDEX_NONE;
+}
+
+FText SBoneItem::GetName() const
+{
+	return FText::FromName(WeakTreeElement.Pin()->BoneName);
+}
+
+void SReferenceSkeletonRow::Construct(const FArguments& InArgs, const TSharedRef<STableViewBase>& InOwnerTable)
+{
+	WeakTreeElement = InArgs._TreeElement;
+	WeakModifier = InArgs._WeakModifier;
+	Delegates = InArgs._Delegates;
+
+	const FSuperRowType::FArguments Args = FSuperRowType::FArguments()
+		.OnDragDetected(InArgs._Delegates.OnDragDetected)
+		.OnCanAcceptDrop(InArgs._Delegates.OnCanAcceptDrop)
+		.OnAcceptDrop(InArgs._Delegates.OnAcceptDrop)
+		.Style( FAppStyle::Get(), "TableView.AlternatingRow" );
+	
+	SMultiColumnTableRow::Construct(Args , InOwnerTable);
+}
+
+TSharedRef<SWidget> SReferenceSkeletonRow::GenerateWidgetForColumn(const FName& ColumnName)
+{
+	return SNew(SHorizontalBox)
+	+ SHorizontalBox::Slot()
+	.AutoWidth()
+	[
+		SNew( SExpanderArrow, SharedThis(this) )
+			.ShouldDrawWires(true)
+	]
+
+	+ SHorizontalBox::Slot()
+	.AutoWidth()
+	[
+		SNew(SBoneItem)
+			.WeakModifier(WeakModifier)
+			.TreeElement(WeakTreeElement)
+			.Delegates(Delegates)
+			.IsSelected(this, &SReferenceSkeletonRow::IsSelected)
+	];
+}
+
+TSharedRef<FBoneItemDragDropOp> FBoneItemDragDropOp::New(TArrayView<const TSharedPtr<FBoneElement>> InElements)
+{
+	TSharedRef<FBoneItemDragDropOp> Operation = MakeShared<FBoneItemDragDropOp>();
+
+	for (const TSharedPtr<FBoneElement>& BoneElement: InElements)
+	{
+		const bool bContainsParent = InElements.Contains(BoneElement->Parent);
+		if (!bContainsParent)
+		{
+			Operation->Elements.Add(BoneElement);
+		}
+	}
+	
+	Operation->Construct();
+	return Operation;
+}
+
+TSharedPtr<SWidget> FBoneItemDragDropOp::GetDefaultDecorator() const
+{
+	return SNew(SBorder)
+        .Visibility(EVisibility::Visible)
+        .BorderImage(FAppStyle::GetBrush("Menu.Background"))
+        [
+            SNew(STextBlock)
+            .Text_Lambda([this]()
+            {
+	            return !Elements.IsEmpty() && Elements[0].IsValid() ? FText::FromName(Elements[0].Pin()->BoneName) : FText::GetEmpty();
+            })
+        ];
+}
+
+SReferenceSkeletonTree::SReferenceSkeletonTree()
+{}
+
+SReferenceSkeletonTree::~SReferenceSkeletonTree()
+{
+	// GEditor->UnregisterForUndo(this);
+}
+
+void SReferenceSkeletonTree::Construct(const FArguments& InArgs)
+{
+	Modifier = InArgs._Modifier;
+	TextFilter = MakeShareable(new FTextFilterExpressionEvaluator(ETextFilterExpressionEvaluatorMode::BasicString));
+	CommandList = MakeShared<FUICommandList>();
+	BindCommands();
+	
+	FRefSkeletonTreeDelegates Delegates;
+	Delegates.OnCanAcceptDrop = FOnRefSkeletonTreeCanAcceptDrop::CreateSP(this, &SReferenceSkeletonTree::OnCanAcceptDrop);
+	Delegates.OnAcceptDrop = FOnRefSkeletonTreeAcceptDrop::CreateSP(this, &SReferenceSkeletonTree::OnAcceptDrop);
+	Delegates.OnDragDetected = FOnDragDetected::CreateSP(this, &SReferenceSkeletonTree::OnDragDetected);
+	Delegates.OnBoneRenamed = FOnBoneRenamed::CreateSP(this, &SReferenceSkeletonTree::OnBoneRenamed);
+	Delegates.OnBoneNameCommitted = FOnTextCommitted::CreateSP(this, &SReferenceSkeletonTree::OnNewBoneNameCommitted);
+	
+	ChildSlot
+	[
+		SNew( SOverlay )
+		+SOverlay::Slot()
+		[
+			SNew(SVerticalBox)
+			+SVerticalBox::Slot()
+			.AutoHeight()
+			.Padding(FMargin(0.f, 2.f))
+			[
+				SNew(SHorizontalBox)
+				+SHorizontalBox::Slot()
+				.AutoWidth()
+				.VAlign(VAlign_Center)
+				.Padding(FMargin(6.f, 0.0))
+				[
+					SNew(SPositiveActionButton)
+					.OnGetMenuContent(this, &SReferenceSkeletonTree::CreateAddNewMenu)
+					.Icon(FAppStyle::Get().GetBrush("Icons.Plus"))
+				]
+
+				+SHorizontalBox::Slot()
+				.FillWidth(1.0f)
+				[
+					SNew(SSearchBox)
+					.SelectAllTextWhenFocused(true)
+					.OnTextChanged( this, &SReferenceSkeletonTree::OnFilterTextChanged )
+					.HintText(LOCTEXT( "SearchBoxHint", "Search Reference Skeleton Tree..."))
+				]
+			]
+			
+			+SVerticalBox::Slot()
+			.Padding(0.0f, 0.0f)
+			[
+				SNew(SBorder)
+				.Padding(2.0f)
+				.BorderImage(FAppStyle::GetBrush("SCSEditor.TreePanel"))
+				[
+					SAssignNew(TreeView, SRefSkeletonTreeView)
+					.TreeItemsSource(&RootElements)
+					.SelectionMode(ESelectionMode::Multi)
+					.OnGenerateRow_Lambda( [this, Delegates](
+						TSharedPtr<FBoneElement> InItem,
+						const TSharedRef<STableViewBase>& OwnerTable) -> TSharedRef<ITableRow>
+					{
+						return SNew(SReferenceSkeletonRow, OwnerTable)
+							.WeakModifier(Modifier)
+							.TreeElement(InItem)
+							.Delegates(Delegates);
+					})
+					.OnGetChildren(this, &SReferenceSkeletonTree::HandleGetChildrenForTree)
+					.OnSelectionChanged(this, &SReferenceSkeletonTree::OnSelectionChanged)
+					.OnContextMenuOpening(this, &SReferenceSkeletonTree::CreateContextMenu)
+					.OnMouseButtonDoubleClick(this, &SReferenceSkeletonTree::OnItemDoubleClicked)
+					.OnSetExpansionRecursive(this, &SReferenceSkeletonTree::OnSetExpansionRecursive)
+					.HighlightParentNodesForSelection(true)
+					.HeaderRow
+					(
+						SNew(SHeaderRow)
+						+ SHeaderRow::Column("Name")
+						.ShouldGenerateWidget(true)
+						.DefaultLabel(LOCTEXT("BoneNameLabel", "Name"))
+						.FillWidth(0.5f)
+					)
+				]
+			]
+		]
+    ];
+
+	static constexpr bool IsInitialSetup = true;
+	RefreshTreeView(IsInitialSetup);
+
+	GEditor->RegisterForUndo(this);
+}
+
+void SReferenceSkeletonTree::PostUndo(bool bSuccess)
+{
+	RefreshTreeView(true);
+}
+
+void SReferenceSkeletonTree::PostRedo(bool bSuccess)
+{
+	RefreshTreeView(true);
+}
+
+void SReferenceSkeletonTree::AddItemToSelection(const TSharedPtr<FBoneElement>& InItem)
+{
+	TreeView->SetItemSelection(InItem, true, ESelectInfo::Direct);
+	TreeView->RequestScrollIntoView(InItem);
+}
+
+void SReferenceSkeletonTree::RemoveItemFromSelection(const TSharedPtr<FBoneElement>& InItem)
+{
+	TreeView->SetItemSelection(InItem, false, ESelectInfo::Direct);
+}
+
+void SReferenceSkeletonTree::ReplaceItemInSelection(const FText& InOldName, const FText& InNewName)
+{
+	const FName OldName(*InOldName.ToString());
+	const FName NewName(*InNewName.ToString());
+	
+	for (const TSharedPtr<FBoneElement>& Item : AllElements)
+	{
+		// remove old selection
+		if (Item->BoneName.IsEqual(OldName))
+		{
+			TreeView->SetItemSelection(Item, false, ESelectInfo::Direct);
+		}
+		// add new selection
+		if (Item->BoneName.IsEqual(NewName))
+		{
+			TreeView->SetItemSelection(Item, true, ESelectInfo::Direct);
+		}
+	}
+}
+
+TArray<TSharedPtr<FBoneElement>> SReferenceSkeletonTree::GetSelectedItems() const
+{
+	return TreeView->GetSelectedItems();
+}
+
+bool SReferenceSkeletonTree::HasSelectedItems() const
+{
+	return TreeView->GetNumItemsSelected() > 0;
+}
+
+void SReferenceSkeletonTree::BindCommands()
+{
+	const FSkeletalMeshModelingToolsCommands& Commands = FSkeletalMeshModelingToolsCommands::Get();
+ 
+	CommandList->MapAction(Commands.NewBone,
+        FExecuteAction::CreateSP(this, &SReferenceSkeletonTree::HandleNewBone),
+        FCanExecuteAction::CreateSP(this, &SReferenceSkeletonTree::CanAddNewBone));
+	
+	CommandList->MapAction(Commands.RemoveBone,
+        FExecuteAction::CreateSP(this, &SReferenceSkeletonTree::HandleDeleteBone),
+        FCanExecuteAction::CreateSP(this, &SReferenceSkeletonTree::CanDeleteBone));
+ 
+	CommandList->MapAction(Commands.RenameBone,
+	FExecuteAction::CreateSP(this, &SReferenceSkeletonTree::HandleRenameBone),
+		FCanExecuteAction::CreateSP(this, &SReferenceSkeletonTree::CanRenameBone));
+	
+	CommandList->MapAction(Commands.UnParentBone,
+		FExecuteAction::CreateSP(this, &SReferenceSkeletonTree::HandleUnParentBone),
+		FCanExecuteAction::CreateSP(this, &SReferenceSkeletonTree::CanUnParentBone));
+
+	CommandList->MapAction(Commands.CopyBones,
+		FExecuteAction::CreateSP(this, &SReferenceSkeletonTree::HandleCopyBones),
+		FCanExecuteAction::CreateSP(this, &SReferenceSkeletonTree::CanCopyBones));
+	
+	CommandList->MapAction(Commands.PasteBones,
+		FExecuteAction::CreateSP(this, &SReferenceSkeletonTree::HandlePasteBones),
+		FCanExecuteAction::CreateSP(this, &SReferenceSkeletonTree::CanPasteBones));
+	
+	CommandList->MapAction(Commands.DuplicateBones,
+		FExecuteAction::CreateSP(this, &SReferenceSkeletonTree::HandleDuplicateBones),
+		FCanExecuteAction::CreateSP(this, &SReferenceSkeletonTree::CanDuplicateBones));
+}
+
+void SReferenceSkeletonTree::HandleNewBone()
+{
+	if (!Modifier.IsValid())
+	{
+		return;
+	}
+
+	BeginChange();
+	
+	TArray<FName> BoneNames;
+	GetSelectedBoneNames(BoneNames);
+	
+	const FName ParentName = BoneNames.IsEmpty() ? NAME_None : BoneNames[0];
+
+	static const FName DefaultName("joint");
+	const FName BoneName = Modifier->GetUniqueName(DefaultName);
+	const bool bAdded = Modifier->AddBone(BoneName, ParentName, FTransform::Identity);
+	if (bAdded)
+	{
+		RefreshTreeView(true);
+		if (Notifier.IsValid())
+		{
+			Notifier->Notify( {BoneName}, ESkeletalMeshNotifyType::BonesAdded);
+		}
+		EndChange();
+	}
+	else
+	{
+		CancelChange();
+	}
+}
+
+bool SReferenceSkeletonTree::CanAddNewBone() const
+{
+	return true;
+}
+
+void SReferenceSkeletonTree::HandleDeleteBone()
+{
+	if (!Modifier.IsValid())
+	{
+		return;
+	}
+
+	BeginChange();
+	
+	TArray<FName> BoneNames;
+	GetSelectedBoneNames(BoneNames);
+	
+	const bool bRemoved = Modifier->RemoveBones(BoneNames, true);
+	if (bRemoved)
+	{
+		RefreshTreeView(true);
+		if (Notifier.IsValid())
+		{
+			Notifier->Notify(BoneNames, ESkeletalMeshNotifyType::BonesRemoved);
+		}
+		EndChange();
+	}
+	else
+	{
+		CancelChange();
+	}
+}
+
+bool SReferenceSkeletonTree::CanDeleteBone() const
+{
+	return HasSelectedItems();
+}
+
+void SReferenceSkeletonTree::HandleUnParentBone()
+{
+	if (!Modifier.IsValid())
+	{
+		return;
+	}
+
+	BeginChange();
+	
+	TArray<FName> BoneNames;
+	GetSelectedBoneNames(BoneNames);
+
+	static const TArray<FName> Dummy;
+	const bool bUnParented = Modifier->ParentBones(BoneNames, Dummy);
+	if (bUnParented)
+	{
+		RefreshTreeView(true);
+		SelectItemFromNames(BoneNames);
+		if (Notifier.IsValid())
+		{
+			Notifier->Notify(BoneNames, ESkeletalMeshNotifyType::HierarchyChanged);
+		}
+		EndChange();
+	}
+	else
+	{
+		CancelChange();
+	}
+}
+
+bool SReferenceSkeletonTree::CanUnParentBone() const
+{
+	return HasSelectedItems();
+}
+
+void SReferenceSkeletonTree::HandleCopyBones() const
+{
+	if (!ensure(Modifier.IsValid()))
+	{
+		return;
+	}
+
+	TArray<FName> BoneNames;
+	GetSelectedBoneNames(BoneNames);
+
+	if (BoneNames.IsEmpty())
+	{
+		return;
+	}
+
+	SkeletonClipboard::CopyToClipboard(*Modifier.Get(), BoneNames);
+}
+
+bool SReferenceSkeletonTree::CanCopyBones() const
+{
+	return Modifier.IsValid() && HasSelectedItems();
+}
+
+void SReferenceSkeletonTree::HandlePasteBones()
+{
+	TArray<FName> BoneNames;
+	GetSelectedBoneNames(BoneNames);
+
+	const FName DefaultParent = BoneNames.IsEmpty() ? NAME_None : BoneNames[0];  
+
+	BeginChange();
+	
+	const TArray<FName> NewBones = SkeletonClipboard::PasteFromClipboard(*Modifier.Get(), DefaultParent);
+	if (NewBones.IsEmpty())
+	{
+		CancelChange();
+		return;
+	}
+
+	static constexpr bool bRebuildAll = true;
+	RefreshTreeView(bRebuildAll);
+	static constexpr bool bFrameSelection = true;
+	SelectItemFromNames(NewBones, bFrameSelection);
+
+	if (Notifier.IsValid())
+	{
+		Notifier->Notify(NewBones, ESkeletalMeshNotifyType::HierarchyChanged);
+		Notifier->Notify(NewBones, ESkeletalMeshNotifyType::BonesSelected);
+	}
+
+	EndChange();
+}
+
+bool SReferenceSkeletonTree::CanPasteBones() const
+{
+	return Modifier.IsValid() && SkeletonClipboard::IsClipboardValid();
+}
+
+void SReferenceSkeletonTree::HandleDuplicateBones()
+{
+	HandleCopyBones();
+
+	if (!SkeletonClipboard::IsClipboardValid())
+	{
+		return;
+	}  
+
+	BeginChange();
+	
+	const TArray<FName> NewBones = SkeletonClipboard::PasteFromClipboard(*Modifier.Get(), NAME_None);
+	if (NewBones.IsEmpty())
+	{
+		CancelChange();
+		return;
+	}
+
+	static constexpr bool bRebuildAll = true;
+	RefreshTreeView(bRebuildAll);
+	static constexpr bool bFrameSelection = true;
+	SelectItemFromNames(NewBones, bFrameSelection);
+
+	if (Notifier.IsValid())
+	{
+		Notifier->Notify(NewBones, ESkeletalMeshNotifyType::HierarchyChanged);
+		Notifier->Notify(NewBones, ESkeletalMeshNotifyType::BonesSelected);
+	}
+
+	EndChange();
+}
+
+bool SReferenceSkeletonTree::CanDuplicateBones() const
+{
+	return CanCopyBones();
+}
+
+void SReferenceSkeletonTree::GetSelectedBoneNames(TArray<FName>& OutSelectedBoneNames) const
+{
+	OutSelectedBoneNames.Reset();
+	const TArray<TSharedPtr<FBoneElement>> BoneElements = TreeView->GetSelectedItems();
+	OutSelectedBoneNames.Reserve(BoneElements.Num());
+
+	Algo::Transform(BoneElements, OutSelectedBoneNames, [](const TSharedPtr<FBoneElement>& BoneElement)
+	{
+		return BoneElement->BoneName;
+	});
+}
+
+void SReferenceSkeletonTree::SelectItemFromNames(const TArray<FName>& InBoneNames, bool bFrameSelection)
+{
+	for (const TSharedPtr<FBoneElement>& Item : AllElements)
+	{
+		const bool bSelect = InBoneNames.Contains(Item->BoneName);
+		TreeView->SetItemSelection(Item, bSelect, ESelectInfo::Direct);
+
+		if (bFrameSelection && bSelect)
+		{
+			TreeView->RequestScrollIntoView(Item);
+			bFrameSelection = false;
+		}
+	}
+}
+
+ISkeletalMeshNotifier& SReferenceSkeletonTree::GetNotifier()
+{
+	if (!Notifier)
+	{
+		Notifier.Reset(new FReferenceSkeletonWidgetNotifier(SharedThis(this)));
+	}
+	return *Notifier;
+}
+
+void SReferenceSkeletonTree::HandleRenameBone() const
+{
+	TArray<TSharedPtr<FBoneElement>> BoneElements = TreeView->GetSelectedItems();
+	if (BoneElements.IsEmpty())
+	{
+		return;
+	}
+	BoneElements[0]->RequestRename();
+}
+
+bool SReferenceSkeletonTree::CanRenameBone() const
+{
+	return HasSelectedItems();
+}
+
+void SReferenceSkeletonTree::OnBoneRenamed(const FName InOldName, const FName InNewName) const
+{
+	if (!Modifier.IsValid())
+	{
+		return;
+	}
+
+	if (Notifier.IsValid())
+	{
+		Notifier->Notify( {InNewName}, ESkeletalMeshNotifyType::BonesRenamed);
+	}
+}
+
+void SReferenceSkeletonTree::OnNewBoneNameCommitted(const FText& InText, ETextCommit::Type InCommitType)
+{
+	TArray<FName> BoneNames;
+	GetSelectedBoneNames(BoneNames);
+	
+	const FName NewName = FName(InText.ToString());
+	
+	if (BoneNames.IsEmpty() || NewName == NAME_None)
+	{
+		return;
+	}
+	
+	if (BoneNames.Num() == 1 && BoneNames[0] == NewName)
+	{
+		return;
+	}
+
+	TArray<int32> BoneIndices;
+	const FReferenceSkeleton& ReferenceSkeleton = Modifier->GetReferenceSkeleton();
+	Algo::Transform(BoneNames, BoneIndices, [&](const FName& BoneName)
+	{
+		return ReferenceSkeleton.FindRawBoneIndex(BoneName);
+	});
+
+	TArray<FName> NewNames;
+	NewNames.Init(NewName, BoneNames.Num());
+	
+	BeginChange();
+
+	const bool bBoneRenamed = Modifier->RenameBones(BoneNames, NewNames);
+	if (bBoneRenamed)
+	{
+		const TArray<FMeshBoneInfo>& BoneInfos = ReferenceSkeleton.GetRawRefBoneInfo();
+		for (int32 Index = 0; Index < BoneIndices.Num(); Index++)
+		{
+			if (BoneIndices[Index] != INDEX_NONE)
+			{
+				const FName OldName = BoneNames[Index];
+				BoneNames[Index] = BoneInfos[BoneIndices[Index]].Name;
+				const int32 ItemIndex = AllElements.IndexOfByPredicate([OldName](const TSharedPtr<FBoneElement>& Item)
+				{
+					return Item->BoneName == OldName;
+				});
+
+				if (ItemIndex != INDEX_NONE)
+				{
+					AllElements[ItemIndex]->BoneName = BoneNames[Index];
+				}
+			}
+		}
+		
+		if (Notifier.IsValid())
+		{
+			Notifier->Notify(BoneNames, ESkeletalMeshNotifyType::BonesRenamed);
+		}
+		
+		EndChange();
+	}
+	else
+	{
+		CancelChange();	
+	}
+}
+
+void SReferenceSkeletonTree::OnFilterTextChanged(const FText& SearchText)
+{
+	TextFilter->SetFilterText(SearchText);
+	RefreshTreeView();
+}
+
+void SReferenceSkeletonTree::RefreshTreeView(bool IsInitialSetup /*=false*/)
+{
+	if (!Modifier.IsValid())
+	{
+		return;
+	}
+	
+	const FReferenceSkeleton& RefSkeleton = Modifier->GetReferenceSkeleton();
+	const TArray<FMeshBoneInfo>& BoneInfos = RefSkeleton.GetRawRefBoneInfo();
+	
+	// reset all tree items
+	RootElements.Reset();
+	AllElements.Reset();
+
+	// validate we have a skeleton to load
+	if (BoneInfos.IsEmpty())
+	{
+		TreeView->RequestTreeRefresh();
+		return;
+	}
+	
+	// record bone element indices
+	TMap<FName, int32> BoneTreeElementIndices;
+
+	auto FilterString = [this](const FString& StringToTest) ->bool
+	{
+		return TextFilter->TestTextFilter(FBasicStringFilterExpressionContext(StringToTest));
+	};
+
+	TMap<int32, int32> PostponedParent;
+	
+	// create all bone elements
+	for (int32 BoneIndex=0; BoneIndex < BoneInfos.Num(); ++BoneIndex)
+	{
+		const FName BoneName = BoneInfos[BoneIndex].Name;
+		TSharedPtr<FBoneElement> BoneElement = MakeShared<FBoneElement>(BoneName, Modifier);
+		const int32 BoneElementIndex = AllElements.Add(BoneElement);
+		BoneTreeElementIndices.Add(BoneName, BoneElementIndex);
+
+		// store pointer to parent (if there is one)
+		const int32 ParentIndex = BoneInfos[BoneIndex].ParentIndex;
+		if (ParentIndex != INDEX_NONE)
+		{
+			// get parent tree element
+			const FName ParentBoneName = BoneInfos[ParentIndex].Name;
+
+			const int32* FoundParent = BoneTreeElementIndices.Find(ParentBoneName);
+			if (FoundParent)
+			{ // set parent info directly
+				const TSharedPtr<FBoneElement> ParentBoneTreeElement = AllElements[*FoundParent];
+				BoneElement->UnFilteredParent = ParentBoneTreeElement;
+			}
+			else
+			{ // postpone as the parent might not have been added to AllElements / BoneTreeElementIndices yet
+				PostponedParent.Emplace(BoneElementIndex, ParentIndex);
+			}
+		}
+		
+		// apply text filter to bones
+		if (!(TextFilter->GetFilterText().IsEmpty() || FilterString(BoneName.ToString())))
+		{
+			BoneElement->bIsHidden = true;
+		}
+	}
+
+	// store pointer to parents that were postponed 
+	for (const auto & [BoneElementIndex, ParentIndex] : PostponedParent)
+	{
+		const FName ParentBoneName = BoneInfos[ParentIndex].Name;
+		const int32* FoundParent = BoneTreeElementIndices.Find(ParentBoneName);
+		if (ensure(FoundParent))
+		{
+			const TSharedPtr<FBoneElement> BoneElement = AllElements[BoneElementIndex];
+			const TSharedPtr<FBoneElement> ParentBoneTreeElement = AllElements[*FoundParent];
+			BoneElement->UnFilteredParent = ParentBoneTreeElement;
+		}
+	}
+	
+	// resolve parent/children pointers on all tree elements, taking into consideration the filter options
+	// (elements are parented to their nearest non-hidden/filtered parent element)
+	for (TSharedPtr<FBoneElement>& Element : AllElements)
+	{
+		if (Element->bIsHidden)
+		{
+			continue;
+		}
+		
+		// find first parent that is not filtered
+		TSharedPtr<FBoneElement> ParentElement = Element->UnFilteredParent;
+		while (true)
+		{
+			if (!ParentElement.IsValid())
+			{
+				break;
+			}
+
+			if (!ParentElement->bIsHidden)
+			{
+				break;
+			}
+
+			ParentElement = ParentElement->UnFilteredParent;
+		}
+
+		if (ParentElement.IsValid())
+		{
+			// store pointer to child on parent
+			ParentElement->Children.Add(Element);
+			// store pointer to parent on child
+			Element->Parent = ParentElement;
+		}
+		else
+		{
+			// has no parent, store a root element
+			RootElements.Add(Element);
+		}
+	}
+
+	// expand all elements upon the initial construction of the tree
+	for (const TSharedPtr<FBoneElement>& RootElement : RootElements)
+	{
+		SetExpansionRecursive(RootElement, false, true);
+	}
+	
+	TreeView->RequestTreeRefresh();
+}
+
+void SReferenceSkeletonTree::HandleGetChildrenForTree(
+	TSharedPtr<FBoneElement> InItem,
+	TArray<TSharedPtr<FBoneElement>>& OutChildren)
+{
+	OutChildren = InItem.Get()->Children;
+}
+
+void SReferenceSkeletonTree::OnSelectionChanged(TSharedPtr<FBoneElement> InItem, ESelectInfo::Type InSelectInfo)
+{
+	if (InSelectInfo == ESelectInfo::Direct)
+	{
+		return;
+	}
+	
+	if (!InItem)
+	{
+		if (Notifier.IsValid())
+		{
+			Notifier->Notify({}, ESkeletalMeshNotifyType::BonesSelected);
+		}
+		return;
+	}
+	
+	TArray<FName> BoneNames;
+	GetSelectedBoneNames(BoneNames);
+	
+	if (Notifier.IsValid())
+	{
+		Notifier->Notify(BoneNames, ESkeletalMeshNotifyType::BonesSelected);
+	}
+}
+
+TSharedRef<SWidget> SReferenceSkeletonTree::CreateAddNewMenu()
+{
+	const FSkeletalMeshModelingToolsCommands& Commands = FSkeletalMeshModelingToolsCommands::Get();
+
+	static constexpr bool bCloseAfterSelection = true;
+	FMenuBuilder MenuBuilder(bCloseAfterSelection, CommandList);
+	
+	MenuBuilder.BeginSection("NewBone", LOCTEXT("AddNewBoneOperations", "Bones"));
+	MenuBuilder.AddMenuEntry(Commands.NewBone);
+	MenuBuilder.EndSection();
+		
+	return MenuBuilder.MakeWidget();
+}
+
+TSharedPtr<SWidget> SReferenceSkeletonTree::CreateContextMenu()
+{
+	static constexpr bool bCloseAfterSelection = true;
+	FMenuBuilder MenuBuilder(bCloseAfterSelection, CommandList);
+
+	const FSkeletalMeshModelingToolsCommands& Commands = FSkeletalMeshModelingToolsCommands::Get();
+	
+	const TArray<TSharedPtr<FBoneElement>> SelectedItems = GetSelectedItems();
+	if (SelectedItems.IsEmpty())
+	{
+		MenuBuilder.BeginSection("NewBone", LOCTEXT("AddNewBoneOperations", "Bones"));
+	
+		MenuBuilder.AddMenuEntry(Commands.NewBone);
+
+		MenuBuilder.EndSection();
+
+		MenuBuilder.BeginSection("CopyPasteBones", LOCTEXT("CopyPasteBonesOperations", "Copy & Paste"));
+
+		MenuBuilder.AddMenuEntry(Commands.PasteBones);
+
+		MenuBuilder.EndSection();
+	}
+	else
+	{
+		MenuBuilder.BeginSection("EditBones", LOCTEXT("EditBonesOperations", "Bones"));
+
+		MenuBuilder.AddMenuEntry(Commands.NewBone);
+		MenuBuilder.AddMenuEntry(Commands.RemoveBone);
+		MenuBuilder.AddMenuEntry(Commands.UnParentBone);
+		MenuBuilder.AddMenuEntry(Commands.RenameBone);
+
+		MenuBuilder.EndSection();
+
+		MenuBuilder.BeginSection("CopyPasteBones", LOCTEXT("CopyPasteBonesOperations", "Copy & Paste"));
+
+		MenuBuilder.AddMenuEntry(Commands.CopyBones);
+		MenuBuilder.AddMenuEntry(Commands.PasteBones);
+		MenuBuilder.AddMenuEntry(Commands.DuplicateBones);
+
+		MenuBuilder.EndSection();
+	}
+	
+	return MenuBuilder.MakeWidget();
+}
+
+void SReferenceSkeletonTree::OnItemDoubleClicked(TSharedPtr<FBoneElement> InItem)
+{
+	InItem->RequestRename();
+}
+
+void SReferenceSkeletonTree::OnSetExpansionRecursive(TSharedPtr<FBoneElement> InItem, bool bShouldBeExpanded)
+{
+	SetExpansionRecursive(InItem, false, bShouldBeExpanded);
+}
+
+void SReferenceSkeletonTree::SetExpansionRecursive(
+	TSharedPtr<FBoneElement> InElement,
+	bool bTowardsParent,
+    bool bShouldBeExpanded)
+{
+	TreeView->SetItemExpansion(InElement, bShouldBeExpanded);
+    
+    if (bTowardsParent)
+    {
+    	if (InElement->Parent.Get())
+    	{
+    		SetExpansionRecursive(InElement->Parent, bTowardsParent, bShouldBeExpanded);
+    	}
+    }
+    else
+    {
+    	for (int32 ChildIndex = 0; ChildIndex < InElement->Children.Num(); ++ChildIndex)
+    	{
+    		SetExpansionRecursive(InElement->Children[ChildIndex], bTowardsParent, bShouldBeExpanded);
+    	}
+    }
+}
+
+FReply SReferenceSkeletonTree::OnKeyDown(const FGeometry& MyGeometry, const FKeyEvent& InKeyEvent)
+{
+	if (!Modifier.IsValid())
+	{
+		return FReply::Handled();
+	}
+	
+	if (CommandList.IsValid() && CommandList->ProcessCommandBindings(InKeyEvent))
+	{
+		return FReply::Handled();
+	}
+	
+	return FReply::Unhandled();
+}
+
+FReply SReferenceSkeletonTree::OnDragDetected(const FGeometry& MyGeometry, const FPointerEvent& MouseEvent)
+{
+	const TArray<TSharedPtr<FBoneElement>> BoneElements = GetSelectedItems();
+	if (BoneElements.IsEmpty())
+	{
+		return FReply::Unhandled();
+	}
+	
+	if (MouseEvent.IsMouseButtonDown(EKeys::LeftMouseButton))
+	{
+		const TSharedRef<FBoneItemDragDropOp> DragDropOp = FBoneItemDragDropOp::New(BoneElements);
+		return FReply::Handled().BeginDragDrop(DragDropOp);
+	}
+	
+	return FReply::Unhandled();
+}
+
+TOptional<EItemDropZone> SReferenceSkeletonTree::OnCanAcceptDrop(const FDragDropEvent& DragDropEvent, EItemDropZone DropZone, TSharedPtr<FBoneElement> TargetItem)
+{
+	TOptional<EItemDropZone> ReturnedDropZone;
+
+	const TSharedPtr<FBoneItemDragDropOp> DragDropOp = DragDropEvent.GetOperationAs<FBoneItemDragDropOp>();
+	if (DragDropOp.IsValid())
+	{
+		ReturnedDropZone = EItemDropZone::OntoItem;
+	}
+
+	return ReturnedDropZone;
+}
+
+FReply SReferenceSkeletonTree::OnAcceptDrop(const FDragDropEvent& DragDropEvent, EItemDropZone DropZone, TSharedPtr<FBoneElement> TargetItem)
+{
+	const TSharedPtr<FBoneItemDragDropOp> DragDropOp = DragDropEvent.GetOperationAs<FBoneItemDragDropOp>();
+	if (!DragDropOp.IsValid())
+	{
+		return FReply::Unhandled();
+	}
+
+	if (!Modifier.IsValid())
+	{
+		return FReply::Handled();
+	}
+	
+	TArray<FName> BoneNames;
+	BoneNames.Reserve(DragDropOp->Elements.Num());
+	for (const TWeakPtr<FBoneElement>& BoneElement: DragDropOp->Elements)
+	{
+		if (BoneElement.IsValid())
+		{
+			BoneNames.Add(BoneElement.Pin()->BoneName);
+		}
+	}
+	const FName ParentName = TargetItem->BoneName;
+
+	BeginChange();
+	
+	const bool bParented = Modifier->ParentBones(BoneNames, {ParentName});
+	if (bParented)
+	{
+		RefreshTreeView(true);
+		SelectItemFromNames(BoneNames);
+		if (Notifier.IsValid())
+		{
+			Notifier->Notify( BoneNames, ESkeletalMeshNotifyType::HierarchyChanged);
+			Notifier->Notify( BoneNames, ESkeletalMeshNotifyType::BonesSelected);
+		}
+		EndChange();
+	}
+	else
+	{
+		CancelChange();
+	}
+	
+	return FReply::Handled();
+}
+
+FReferenceSkeletonWidgetNotifier::FReferenceSkeletonWidgetNotifier(TSharedRef<SReferenceSkeletonTree> InWidget)
+	: ISkeletalMeshNotifier()
+	, Tree(InWidget)
+{}
+
+void FReferenceSkeletonWidgetNotifier::HandleNotification(const TArray<FName>& InBoneNames, const ESkeletalMeshNotifyType InNotifyType)
+{
+	if (Notifying() || !Tree.IsValid())
+	{
+		return;
+	}
+
+	TSharedPtr<SReferenceSkeletonTree> TreePtr = Tree.Pin();
+	switch (InNotifyType)
+	{
+		case ESkeletalMeshNotifyType::BonesAdded:
+			TreePtr->RefreshTreeView();
+			break;
+		case ESkeletalMeshNotifyType::BonesRemoved:
+			TreePtr->RefreshTreeView();
+			break;
+		case ESkeletalMeshNotifyType::BonesMoved:
+			break;
+		case ESkeletalMeshNotifyType::BonesSelected:
+			{
+				static constexpr bool bFrameSelection = true;
+				TreePtr->SelectItemFromNames(InBoneNames, bFrameSelection);
+				break;
+			}
+		case ESkeletalMeshNotifyType::BonesRenamed:
+			TreePtr->RefreshTreeView();
+			TreePtr->SelectItemFromNames(InBoneNames);
+			break;
+		case ESkeletalMeshNotifyType::HierarchyChanged:
+			{
+				TArray<FName> BoneNames;
+				TreePtr->GetSelectedBoneNames(BoneNames);
+				TreePtr->RefreshTreeView();
+				TreePtr->SelectItemFromNames(BoneNames);
+			}
+			break;
+	}
+}
+
+void SReferenceSkeletonTree::BeginChange()
+{
+	ensure( ActiveChange == nullptr );
+	ActiveChange = MakeUnique<ReferenceSkeletonTreeLocals::FSkeletonModifierChange>(Modifier.Get()); 
+}
+
+void SReferenceSkeletonTree::EndChange()
+{
+	if (!ActiveChange.IsValid())
+	{
+		return;
+	}
+
+	ActiveChange->StoreSkeleton(Modifier.Get());
+	
+	GEditor->BeginTransaction(LOCTEXT("ModifyRefSkeleton", "Modify Reference Skeleton"));
+	GUndo->StoreUndo(Modifier.Get(), MoveTemp(ActiveChange));
+	GEditor->EndTransaction();
+}
+
+void SReferenceSkeletonTree::CancelChange()
+{
+	ActiveChange.Reset();
+}
+
+#undef LOCTEXT_NAMESPACE
\ No newline at end of file
diff --git a/Engine/Plugins/Animation/SkeletalMeshModelingTools/Source/SkeletalMeshModelingTools/Private/SReferenceSkeletonTree.h b/Engine/Plugins/Animation/SkeletalMeshModelingTools/Source/SkeletalMeshModelingTools/Private/SReferenceSkeletonTree.h
new file mode 100644
--- /dev/null
+++ b/Engine/Plugins/Animation/SkeletalMeshModelingTools/Source/SkeletalMeshModelingTools/Private/SReferenceSkeletonTree.h
@@ -0,0 +1,259 @@
+// Copyright Epic Games, Inc. All Rights Reserved.
+
+#pragma once
+
+#include "EditorUndoClient.h"
+#include "SkeletalMeshNotifier.h"
+#include "DragAndDrop/DecoratedDragDropOp.h"
+#include "Widgets/SCompoundWidget.h"
+#include "Widgets/Views/SListView.h"
+#include "Widgets/Views/STableRow.h"
+#include "Widgets/Views/STreeView.h"
+#include "Misc/Change.h"
+#include "ReferenceSkeleton.h"
+
+class FTextFilterExpressionEvaluator;
+class FUICommandList;
+class USkeletonModifier;
+
+namespace ReferenceSkeletonTreeLocals
+{
+
+class FSkeletonModifierChange : public FCommandChange
+{
+public:
+	FSkeletonModifierChange(const USkeletonModifier* InModifier);
+
+	virtual FString ToString() const override
+	{
+		return FString(TEXT("Edit Skeleton"));
+	}
+
+	virtual void Apply(UObject* Object) override;
+	virtual void Revert(UObject* Object) override;
+
+	void StoreSkeleton(const USkeletonModifier* InModifier);
+
+private:
+	FReferenceSkeleton PreChangeSkeleton;
+	TArray<int32> PreBoneTracker;
+	FReferenceSkeleton PostChangeSkeleton;
+	TArray<int32> PostBoneTracker;
+};
+
+}
+
+class FBoneElement : public TSharedFromThis<FBoneElement>
+{
+public:
+	FBoneElement(const FName& InBoneName, TWeakObjectPtr<USkeletonModifier> InModifier);
+
+	FName BoneName;
+
+	bool bIsHidden = false;
+	
+	TSharedPtr<FBoneElement> UnFilteredParent;
+	TSharedPtr<FBoneElement> Parent;
+	TArray<TSharedPtr<FBoneElement>> Children;
+
+	void RequestRename() const;
+	
+	DECLARE_DELEGATE(FOnRenameRequested);
+	FOnRenameRequested OnRenameRequested;
+
+	DECLARE_DELEGATE_TwoParams(FOnRenamed, const FName, const FName);
+	FOnRenamed OnRenamed;
+
+private:
+	TWeakObjectPtr<USkeletonModifier> WeakModifier;
+};
+
+using FOnRefSkeletonTreeCanAcceptDrop = STableRow<TSharedPtr<FBoneElement>>::FOnCanAcceptDrop;
+using FOnRefSkeletonTreeAcceptDrop = STableRow<TSharedPtr<FBoneElement>>::FOnAcceptDrop;
+using FOnBoneRenamed = FBoneElement::FOnRenamed;
+
+struct FRefSkeletonTreeDelegates
+{
+	FOnDragDetected					OnDragDetected;
+	FOnRefSkeletonTreeCanAcceptDrop	OnCanAcceptDrop;
+	FOnRefSkeletonTreeAcceptDrop	OnAcceptDrop;
+	FOnBoneRenamed					OnBoneRenamed;
+	FOnTextCommitted				OnBoneNameCommitted;
+};
+
+class SBoneItem : public SCompoundWidget
+{
+	SLATE_BEGIN_ARGS(SBoneItem) {}
+		SLATE_ARGUMENT(TWeakPtr<FBoneElement>, TreeElement)
+		SLATE_ARGUMENT(TWeakObjectPtr<USkeletonModifier>, WeakModifier)
+		SLATE_ARGUMENT(FRefSkeletonTreeDelegates, Delegates)
+		SLATE_EVENT( FIsSelected, IsSelected )
+	SLATE_END_ARGS()
+
+	// Slate construction function
+	void Construct(const FArguments& InArgs);
+	
+private:
+
+	bool OnVerifyNameChanged(const FText& InText, FText& OutErrorMessage) const;
+
+	static FSlateColor GetTextColor(FIsSelected InIsSelected);
+	
+	FText GetName() const;
+	
+	TWeakPtr<FBoneElement> WeakTreeElement;
+	TWeakObjectPtr<USkeletonModifier> WeakModifier;
+};
+
+class FBoneItemDragDropOp : public FDecoratedDragDropOp
+{
+public:
+	DRAG_DROP_OPERATOR_TYPE(FBoneItemDragDropOp, FDecoratedDragDropOp)
+
+	static TSharedRef<FBoneItemDragDropOp> New(TArrayView<const TSharedPtr<FBoneElement>> InElements);
+
+	// FDecoratedDragDropOp interface
+	virtual TSharedPtr<SWidget> GetDefaultDecorator() const override;
+	// End FDecoratedDragDropOp interface
+	
+	TArray< TWeakPtr<FBoneElement> > Elements;
+};
+
+class SReferenceSkeletonRow : public SMultiColumnTableRow<TSharedPtr<FBoneElement>>
+{
+public:
+	SLATE_BEGIN_ARGS(SReferenceSkeletonRow) {}
+		SLATE_ARGUMENT(TWeakObjectPtr<USkeletonModifier>, WeakModifier)
+		SLATE_ARGUMENT(TSharedPtr<FBoneElement>, TreeElement)
+		SLATE_ARGUMENT(FRefSkeletonTreeDelegates, Delegates)
+	SLATE_END_ARGS()
+
+	// Slate construction function
+	void Construct(const FArguments& InArgs, const TSharedRef<STableViewBase>& InOwnerTable);
+
+	// SMultiColumnTableRow interface
+	virtual TSharedRef<SWidget> GenerateWidgetForColumn(const FName& InColumnName) override;
+	// End SMultiColumnTableRow interface
+
+private:
+	TWeakPtr<FBoneElement> WeakTreeElement;
+	TWeakObjectPtr<USkeletonModifier> WeakModifier;
+	FRefSkeletonTreeDelegates Delegates;
+};
+
+typedef STreeView< TSharedPtr<FBoneElement> > SRefSkeletonTreeView;
+
+class SReferenceSkeletonTree;
+
+class FReferenceSkeletonWidgetNotifier: public ISkeletalMeshNotifier
+{
+public:
+	FReferenceSkeletonWidgetNotifier(TSharedRef<SReferenceSkeletonTree> InTree);
+	virtual void HandleNotification(const TArray<FName>& BoneNames, const ESkeletalMeshNotifyType InNotifyType) override;
+	
+private:
+	TWeakPtr<SReferenceSkeletonTree> Tree;
+};
+
+class SReferenceSkeletonTree : public SCompoundWidget, public FEditorUndoClient
+{
+public:
+	SLATE_BEGIN_ARGS(SReferenceSkeletonTree) {}
+		SLATE_ARGUMENT(TWeakObjectPtr<USkeletonModifier>, Modifier)
+		SLATE_ARGUMENT(FRefSkeletonTreeDelegates, Delegates)
+	SLATE_END_ARGS()
+
+	SReferenceSkeletonTree();
+    virtual ~SReferenceSkeletonTree() override;
+
+	// Slate construction function
+	void Construct(const FArguments& InArgs);
+
+	virtual void PostUndo(bool bSuccess) override;
+	virtual void PostRedo(bool bSuccess) override;
+	
+	// Selection
+	void GetSelectedBoneNames(TArray<FName>& OutSelectedBoneNames) const;
+	TArray<TSharedPtr<FBoneElement>> GetSelectedItems() const;
+	bool HasSelectedItems() const;
+	void SelectItemFromNames(const TArray<FName>& InBoneNames, bool bFrameSelection = false);
+
+	ISkeletalMeshNotifier& GetNotifier();
+
+private:
+	// SWidget interface
+	virtual FReply OnKeyDown(const FGeometry& MyGeometry, const FKeyEvent& InKeyEvent) override;
+	// END SWidget interface
+	
+	void BindCommands();
+
+	void AddItemToSelection(const TSharedPtr<FBoneElement>& InItem);
+	void RemoveItemFromSelection(const TSharedPtr<FBoneElement>& InItem);
+	void ReplaceItemInSelection(const FText& OldName, const FText& NewName);
+
+	// Create, Rename, Delete, Parent bones
+	void HandleNewBone();
+	bool CanAddNewBone() const;
+	void HandleRenameBone() const;
+	bool CanRenameBone() const;
+	void HandleDeleteBone();
+	bool CanDeleteBone() const;
+	void HandleUnParentBone();
+	bool CanUnParentBone() const;
+
+	// Copy & Paste
+	void HandleCopyBones() const;
+	bool CanCopyBones() const;
+	void HandlePasteBones();
+	bool CanPasteBones() const;
+	void HandleDuplicateBones();
+	bool CanDuplicateBones() const;
+
+	// Filter
+	void OnFilterTextChanged(const FText& SearchText);
+	
+	// Callbacks
+	void RefreshTreeView(bool IsInitialSetup=false);
+	void HandleGetChildrenForTree(TSharedPtr<FBoneElement> InItem, TArray<TSharedPtr<FBoneElement>>& OutChildren);
+	void OnSelectionChanged(TSharedPtr<FBoneElement> InItem, ESelectInfo::Type InSelectInfo);
+	
+	TSharedRef< SWidget > CreateAddNewMenu();
+	TSharedPtr< SWidget > CreateContextMenu();
+
+	void OnItemDoubleClicked(TSharedPtr<FBoneElement> InItem);
+	void OnSetExpansionRecursive(TSharedPtr<FBoneElement> InItem, bool bShouldBeExpanded);
+	void SetExpansionRecursive(TSharedPtr<FBoneElement> InElement, bool bTowardsParent, bool bShouldBeExpanded);
+
+	// Delegates
+	virtual FReply OnDragDetected(const FGeometry& MyGeometry, const FPointerEvent& MouseEvent) override;
+	TOptional<EItemDropZone> OnCanAcceptDrop(const FDragDropEvent& DragDropEvent, EItemDropZone DropZone, TSharedPtr<FBoneElement> TargetItem);
+	FReply OnAcceptDrop(const FDragDropEvent& DragDropEvent, EItemDropZone DropZone, TSharedPtr<FBoneElement> TargetItem);
+	void OnBoneRenamed(const FName InOldName, const FName InNewName) const;
+	
+	void OnNewBoneNameCommitted(const FText& InText, ETextCommit::Type InCommitType);
+
+	// The modifier that actually does modifications on the reference skeleton
+	TWeakObjectPtr<USkeletonModifier> Modifier;
+	
+	TUniquePtr<FReferenceSkeletonWidgetNotifier> Notifier;
+
+	// undo
+	void BeginChange();
+	void EndChange();
+	void CancelChange();
+	TUniquePtr<ReferenceSkeletonTreeLocals::FSkeletonModifierChange> ActiveChange;
+	
+	// Commands
+	TSharedPtr<FUICommandList> CommandList;
+	
+	// Tree view
+	TSharedPtr<SRefSkeletonTreeView> TreeView;
+	TArray<TSharedPtr<FBoneElement>> RootElements;
+	TArray<TSharedPtr<FBoneElement>> AllElements;
+
+	// filter
+	TSharedPtr<FTextFilterExpressionEvaluator> TextFilter;
+
+	friend class FReferenceSkeletonWidgetNotifier;
+};
+
diff --git a/Engine/Plugins/Animation/SkeletalMeshModelingTools/Source/SkeletalMeshModelingTools/Private/SkeletalMeshEditorUtils.cpp b/Engine/Plugins/Animation/SkeletalMeshModelingTools/Source/SkeletalMeshModelingTools/Private/SkeletalMeshEditorUtils.cpp
new file mode 100644
--- /dev/null
+++ b/Engine/Plugins/Animation/SkeletalMeshModelingTools/Source/SkeletalMeshModelingTools/Private/SkeletalMeshEditorUtils.cpp
@@ -0,0 +1,332 @@
+﻿// Copyright Epic Games, Inc. All Rights Reserved.
+
+#include "SkeletalMeshEditorUtils.h"
+
+#include "ContextObjectStore.h"
+#include "InteractiveToolsContext.h"
+#include "IPersonaToolkit.h"
+#include "ISkeletalMeshEditor.h"
+#include "ISkeletonEditorModule.h"
+#include "SkeletonModifier.h"
+#include "SReferenceSkeletonTree.h"
+#include "Widgets/Docking/SDockTab.h"
+#include "Preferences/PersonaOptions.h"
+
+#include UE_INLINE_GENERATED_CPP_BY_NAME(SkeletalMeshEditorUtils)
+
+bool UE::SkeletalMeshEditorUtils::RegisterEditorContextObject(UInteractiveToolsContext* ToolsContext)
+{
+	if (ensure(ToolsContext))
+	{
+		const USkeletalMeshEditorContextObject* Found = ToolsContext->ContextObjectStore->FindContext<USkeletalMeshEditorContextObject>();
+		if (Found)
+		{
+			return true;
+		}
+		
+		USkeletalMeshEditorContextObject* ContextObject = NewObject<USkeletalMeshEditorContextObject>(ToolsContext->ToolManager);
+		if (ensure(ContextObject))
+		{
+			ContextObject->Register(ToolsContext->ToolManager);
+			return true;
+		}
+	}
+	return false;
+}
+
+bool UE::SkeletalMeshEditorUtils::UnregisterEditorContextObject(UInteractiveToolsContext* ToolsContext)
+{
+	if (ensure(ToolsContext))
+	{
+		USkeletalMeshEditorContextObject* Found = ToolsContext->ContextObjectStore->FindContext<USkeletalMeshEditorContextObject>();
+		if (Found != nullptr)
+		{
+			Found->Unregister(ToolsContext->ToolManager);
+			ToolsContext->ContextObjectStore->RemoveContextObject(Found);
+		}
+		return true;
+	}
+	return false;
+}
+
+USkeletalMeshEditorContextObject* UE::SkeletalMeshEditorUtils::GetEditorContextObject(UInteractiveToolsContext* ToolsContext)
+{
+	return ToolsContext->ContextObjectStore->FindContext<USkeletalMeshEditorContextObject>();
+}
+
+void USkeletalMeshEditorContextObject::Register(UInteractiveToolManager* InToolManager)
+{
+	if (ensure(!bRegistered) == false)
+	{
+		return;
+	}
+
+	InToolManager->GetContextObjectStore()->AddContextObject(this);
+	bRegistered = true;
+}
+
+void USkeletalMeshEditorContextObject::Unregister(UInteractiveToolManager* InToolManager)
+{
+	ensure(bRegistered);
+	
+	InToolManager->GetContextObjectStore()->RemoveContextObject(this);
+
+	EditorBindings.Reset();
+	TreeBindings.Reset();
+	
+	bRegistered = false;
+}
+
+void USkeletalMeshEditorContextObject::Init(const TWeakPtr<ISkeletalMeshEditor>& InEditor)
+{
+	Editor = InEditor;
+	EditorBindings.Reset();
+	TreeBindings.Reset();
+}
+
+void USkeletalMeshEditorContextObject::HideSkeleton()
+{
+	if (!Editor.IsValid())
+	{
+		return;
+	}
+
+	const TSharedPtr<ISkeletalMeshEditor> SkeletalMeshEditor = Editor.Pin();
+	UDebugSkelMeshComponent* SkeletalMeshComponent = SkeletalMeshEditor->GetPersonaToolkit()->GetPreviewMeshComponent();
+	if (!SkeletalMeshComponent)
+	{
+		return;
+	}
+	
+	SkeletonDrawMode = SkeletalMeshComponent->SkeletonDrawMode;
+	SkeletalMeshComponent->SkeletonDrawMode = ESkeletonDrawMode::Hidden;
+}
+
+void USkeletalMeshEditorContextObject::ShowSkeleton()
+{
+	if (!Editor.IsValid())
+	{
+		return;
+	}
+	
+	const TSharedPtr<ISkeletalMeshEditor> SkeletalMeshEditor = Editor.Pin();
+	UDebugSkelMeshComponent* SkeletalMeshComponent = SkeletalMeshEditor->GetPersonaToolkit()->GetPreviewMeshComponent();
+	if (!SkeletalMeshComponent)
+	{
+		return;
+	}
+	
+	SkeletalMeshComponent->SkeletonDrawMode = SkeletonDrawMode;
+	SkeletonDrawMode = ESkeletonDrawMode::Default;
+}
+
+void USkeletalMeshEditorContextObject::BindTo(ISkeletalMeshEditingInterface* InEditingInterface)
+{
+	if (!InEditingInterface)
+	{
+		return;
+	}
+	
+	BindEditor(InEditingInterface);
+	BindRefSkeletonTree(InEditingInterface);
+}
+
+void USkeletalMeshEditorContextObject::UnbindFrom(ISkeletalMeshEditingInterface* InEditingInterface)
+{
+	if (!InEditingInterface)
+	{
+		return;
+	}
+
+	UnbindEditor(InEditingInterface);
+	UnbindRefSkeletonTree(InEditingInterface);
+}
+
+TPair<FDelegateHandle, FDelegateHandle> USkeletalMeshEditorContextObject::BindInterfaceTo(
+	ISkeletalMeshEditingInterface* InInterface,
+	ISkeletalMeshNotifier& InOtherNotifier)
+{
+	// connect external interface to tool (ie skeletal mesh editor -> tool)
+	FDelegateHandle ToToolNotifierHandle = InOtherNotifier.Delegate().AddLambda(
+	[InInterface](const TArray<FName>& BoneNames, const ESkeletalMeshNotifyType InNotifyType)
+	{
+		InInterface->GetNotifier().HandleNotification(BoneNames, InNotifyType);
+	});
+
+	// connect tool to external interface (ie tool -> skeletal mesh editor)
+	FDelegateHandle FromToolNotifierHandle = InInterface->GetNotifier().Delegate().AddLambda(
+		[&InOtherNotifier](const TArray<FName>& BoneNames, const ESkeletalMeshNotifyType InNotifyType)
+	{
+		InOtherNotifier.HandleNotification(BoneNames, InNotifyType);
+	});
+
+	return {ToToolNotifierHandle, FromToolNotifierHandle};
+}
+
+void USkeletalMeshEditorContextObject::UnbindInterfaceFrom(
+	ISkeletalMeshEditingInterface* InInterface,
+	ISkeletalMeshNotifier& InOtherNotifier,
+	const FBindData& InOutBindData)
+{
+	if (InOutBindData.ToToolNotifierHandle.IsValid())
+	{
+		InOtherNotifier.Delegate().Remove(InOutBindData.ToToolNotifierHandle);
+	}
+
+	if (InOutBindData.FromToolNotifierHandle.IsValid())
+	{
+		InInterface->GetNotifier().Delegate().Remove(InOutBindData.FromToolNotifierHandle);
+	}
+}
+
+void USkeletalMeshEditorContextObject::BindEditor(ISkeletalMeshEditingInterface* InEditingInterface)
+{
+	if (!InEditingInterface || !Editor.IsValid())
+	{
+		return;
+	}
+	
+	if (EditorBindings.Contains(InEditingInterface))
+	{
+		return;
+	}
+
+	TSharedPtr<ISkeletalMeshEditorBinding> Binding = Editor.Pin()->GetBinding();
+	if (!Binding.IsValid())
+	{
+		return;
+	}
+
+	InEditingInterface->BindTo(Binding);
+
+	FBindData BindData;
+	Tie(BindData.ToToolNotifierHandle, BindData.FromToolNotifierHandle) = BindInterfaceTo(InEditingInterface, Binding->GetNotifier());
+	EditorBindings.Emplace(InEditingInterface, BindData);
+
+	InEditingInterface->GetNotifier().HandleNotification(Binding->GetSelectedBones(), ESkeletalMeshNotifyType::BonesSelected);
+}
+
+void USkeletalMeshEditorContextObject::UnbindEditor(ISkeletalMeshEditingInterface* InEditingInterface)
+{
+	if (const FBindData* BindData = EditorBindings.Find(InEditingInterface))
+	{
+		if (BindData->FromToolNotifierHandle.IsValid())
+		{
+			InEditingInterface->Unbind();
+		}
+
+		if (Editor.IsValid())
+		{
+			const TSharedPtr<ISkeletalMeshEditorBinding> Binding = Editor.Pin()->GetBinding();
+			UnbindInterfaceFrom(InEditingInterface, Binding->GetNotifier(), *BindData);
+		}
+	
+		EditorBindings.Remove(InEditingInterface);
+	}
+}
+
+void USkeletalMeshEditorContextObject::BindRefSkeletonTree(ISkeletalMeshEditingInterface* InEditingInterface)
+{
+	if (!InEditingInterface || !Editor.IsValid())
+	{
+		return;
+	}
+	
+	const TWeakObjectPtr<USkeletonModifier> Modifier = InEditingInterface->GetModifier();
+	if (!Modifier.IsValid())
+	{
+		return;
+	}
+	
+	const TSharedPtr<FTabManager> TabManager = Editor.Pin()->GetAssociatedTabManager();
+	TSharedPtr<SDockTab> SkeletonTab = TabManager->FindExistingLiveTab(GetSkeletonTreeTabId());
+	if (SkeletonTab.IsValid())
+	{
+		DefaultSkeletonWidget = SkeletonTab->GetContent();
+	}
+	else
+	{
+		SkeletonTab = TabManager->TryInvokeTab(GetSkeletonTreeTabId());
+	}
+
+	check(SkeletonTab.IsValid());
+
+	SAssignNew(RefSkeletonWidget, SBorder)
+		.Padding(4.f)
+		.BorderImage(FAppStyle::GetBrush("ToolPanel.GroupBorder"))
+	[
+		SAssignNew(RefSkeletonTree, SReferenceSkeletonTree)
+			.Modifier(Modifier)
+	];
+
+	// switch between SSkeletonTree and SReferenceSkeletonTree
+	SkeletonTab->SetContent(RefSkeletonWidget.ToSharedRef());
+
+	FBindData BindData;
+	Tie(BindData.ToToolNotifierHandle, BindData.FromToolNotifierHandle) = BindInterfaceTo(InEditingInterface, RefSkeletonTree->GetNotifier());
+	TreeBindings.Emplace(InEditingInterface, BindData);
+
+	// unbind notifications on the initial SSkeletonTree to avoid useless notifications
+	if(const TSharedPtr<ISkeletalMeshEditorBinding> Binding = Editor.Pin()->GetBinding())
+	{
+		const TArray<FName> SelectedBones = Binding->GetSelectedBones();
+		
+		if (const FBindData* EditorBindData = EditorBindings.Find(InEditingInterface))
+		{
+			UnbindInterfaceFrom(InEditingInterface, Binding->GetNotifier(), *EditorBindData);
+			Binding->GetNotifier().HandleNotification({}, ESkeletalMeshNotifyType::BonesSelected);
+		}
+
+		RefSkeletonTree->GetNotifier().HandleNotification(SelectedBones, ESkeletalMeshNotifyType::BonesSelected);
+	}
+}
+
+void USkeletalMeshEditorContextObject::UnbindRefSkeletonTree(ISkeletalMeshEditingInterface* InEditingInterface)
+{
+	if (const FBindData* BindData = TreeBindings.Find(InEditingInterface))
+	{
+		TArray<FName> SelectedBones;
+		if (RefSkeletonTree.IsValid())
+		{
+			RefSkeletonTree->GetSelectedBoneNames(SelectedBones);
+			UnbindInterfaceFrom(InEditingInterface, RefSkeletonTree->GetNotifier(), *BindData);
+		}
+
+		if (Editor.IsValid())
+		{
+			const TSharedPtr<FTabManager> TabManager = Editor.Pin()->GetAssociatedTabManager();
+			const TSharedPtr<SDockTab> SkeletonTab = TabManager->FindExistingLiveTab(GetSkeletonTreeTabId());
+			if (SkeletonTab.IsValid())
+			{
+				if (SkeletonTab->GetContent() == RefSkeletonWidget)
+				{
+					SkeletonTab->SetContent(DefaultSkeletonWidget.IsValid() ? DefaultSkeletonWidget.ToSharedRef() : SNullWidget::NullWidget);
+				}
+			}
+
+			// update the SSkeletonTree with the new RefSkeleton and set current selection 
+			if (const TSharedPtr<ISkeletalMeshEditorBinding> Binding = Editor.Pin()->GetBinding())
+			{
+				ISkeletalMeshNotifier& Notifier = Binding->GetNotifier();
+				Notifier.HandleNotification(SelectedBones, ESkeletalMeshNotifyType::HierarchyChanged);
+				Notifier.HandleNotification(SelectedBones, ESkeletalMeshNotifyType::BonesSelected);
+			}
+		}
+		
+		RefSkeletonTree.Reset();
+		RefSkeletonWidget.Reset();
+		
+		TreeBindings.Remove(InEditingInterface);
+	}
+}
+
+const FName& USkeletalMeshEditorContextObject::GetSkeletonTreeTabId()
+{
+	static const FName SkeletonTreeId(TEXT("SkeletonTreeView"));
+	return SkeletonTreeId;
+}
+
+TSharedPtr<ISkeletalMeshEditorBinding> USkeletalMeshEditorContextObject::GetBinding() const
+{
+	return Editor.IsValid() ? Editor.Pin()->GetBinding() : nullptr;
+}
diff --git a/Engine/Plugins/Animation/SkeletalMeshModelingTools/Source/SkeletalMeshModelingTools/Private/SkeletalMeshEditorUtils.h b/Engine/Plugins/Animation/SkeletalMeshModelingTools/Source/SkeletalMeshModelingTools/Private/SkeletalMeshEditorUtils.h
new file mode 100644
--- /dev/null
+++ b/Engine/Plugins/Animation/SkeletalMeshModelingTools/Source/SkeletalMeshModelingTools/Private/SkeletalMeshEditorUtils.h
@@ -0,0 +1,91 @@
+﻿// Copyright Epic Games, Inc. All Rights Reserved.
+
+#pragma once
+
+#include "Animation/DebugSkelMeshComponent.h"
+#include "SkeletalMesh/SkeletalMeshEditionInterface.h"
+
+#include "UObject/Object.h"
+#include "UObject/WeakInterfacePtr.h"
+
+#include "SkeletalMeshEditorUtils.generated.h"
+
+class UInteractiveToolsContext;
+class UInteractiveToolManager;
+class ISkeletalMeshEditor;
+class USkeletalMeshEditorContextObject;
+class SReferenceSkeletonTree;
+
+namespace UE
+{
+	
+namespace SkeletalMeshEditorUtils
+{
+	SKELETALMESHMODELINGTOOLS_API bool RegisterEditorContextObject(UInteractiveToolsContext* InToolsContext);
+	SKELETALMESHMODELINGTOOLS_API bool UnregisterEditorContextObject(UInteractiveToolsContext* InToolsContext);
+
+	SKELETALMESHMODELINGTOOLS_API USkeletalMeshEditorContextObject* GetEditorContextObject(UInteractiveToolsContext* InToolsContext);
+}
+	
+}
+
+/**
+ * USkeletalMeshEditorContextObject
+ */
+
+UCLASS()
+class SKELETALMESHMODELINGTOOLS_API USkeletalMeshEditorContextObject : public USkeletalMeshEditorContextObjectBase
+{
+	GENERATED_BODY()
+
+public:
+	void Register(UInteractiveToolManager* InToolManager);
+	void Unregister(UInteractiveToolManager* InToolManager);
+
+	void Init(const TWeakPtr<ISkeletalMeshEditor>& InEditor);
+
+	virtual void HideSkeleton() override;
+	virtual void ShowSkeleton() override;
+
+	virtual void BindTo(ISkeletalMeshEditingInterface* InEditingInterface) override;
+	virtual void UnbindFrom(ISkeletalMeshEditingInterface* InEditingInterface) override;
+	
+private:
+	TWeakPtr<ISkeletalMeshEditor> Editor = nullptr;
+
+	UPROPERTY()
+	ESkeletonDrawMode SkeletonDrawMode = ESkeletonDrawMode::Default;
+
+	TSharedPtr<ISkeletalMeshEditorBinding> GetBinding() const;
+
+	struct FBindData
+	{
+		FDelegateHandle ToToolNotifierHandle;
+		FDelegateHandle FromToolNotifierHandle;		
+	};
+	TMap< ISkeletalMeshEditingInterface*, FBindData > EditorBindings;
+	TMap< ISkeletalMeshEditingInterface*, FBindData > TreeBindings;
+
+	static TPair<FDelegateHandle, FDelegateHandle> BindInterfaceTo(
+		ISkeletalMeshEditingInterface* InInterface,
+		ISkeletalMeshNotifier& InOtherNotifier);
+	
+	static void UnbindInterfaceFrom(
+		ISkeletalMeshEditingInterface* InInterface,
+		ISkeletalMeshNotifier& InOtherNotifier,
+		const FBindData& InOutBindData);
+
+	void BindEditor(ISkeletalMeshEditingInterface* InEditingInterface);
+	void UnbindEditor(ISkeletalMeshEditingInterface* InEditingInterface);
+
+	void BindRefSkeletonTree(ISkeletalMeshEditingInterface* InEditingInterface);
+	void UnbindRefSkeletonTree(ISkeletalMeshEditingInterface* InEditingInterface);
+
+	static const FName& GetSkeletonTreeTabId();
+
+	TSharedPtr<SWidget> RefSkeletonWidget;
+	TSharedPtr<SReferenceSkeletonTree> RefSkeletonTree;
+	TSharedPtr<SWidget> DefaultSkeletonWidget;
+	
+	bool bRegistered = false;
+};
diff --git a/Engine/Plugins/Animation/SkeletalMeshModelingTools/Source/SkeletalMeshModelingTools/Private/SkeletalMeshGizmoUtils.cpp b/Engine/Plugins/Animation/SkeletalMeshModelingTools/Source/SkeletalMeshModelingTools/Private/SkeletalMeshGizmoUtils.cpp
new file mode 100644
--- /dev/null
+++ b/Engine/Plugins/Animation/SkeletalMeshModelingTools/Source/SkeletalMeshModelingTools/Private/SkeletalMeshGizmoUtils.cpp
@@ -0,0 +1,250 @@
+﻿// Copyright Epic Games, Inc. All Rights Reserved.
+
+#include "SkeletalMeshGizmoUtils.h"
+
+#include "ContextObjectStore.h"
+#include "InteractiveToolsContext.h"
+#include "BaseGizmos/GizmoElementHitTargets.h"
+#include "SkeletalMesh/SkeletonTransformProxy.h"
+
+#include "EditorGizmos/TransformGizmo.h"
+#include "Editor/Experimental/EditorInteractiveToolsFramework/Public/EditorInteractiveGizmoManager.h"
+#include "EditorGizmos/EditorTransformGizmoBuilder.h"
+
+#include "Materials/Material.h"
+
+#include UE_INLINE_GENERATED_CPP_BY_NAME(SkeletalMeshGizmoUtils)
+
+bool UE::SkeletalMeshGizmoUtils::RegisterTransformGizmoContextObject(UInteractiveToolsContext* ToolsContext)
+{
+	if (ensure(ToolsContext))
+	{
+		const USkeletalMeshGizmoContextObject* Found = ToolsContext->ContextObjectStore->FindContext<USkeletalMeshGizmoContextObject>();
+		if (Found)
+		{
+			return true;
+		}
+		
+		USkeletalMeshGizmoContextObject* GizmoContextObject = NewObject<USkeletalMeshGizmoContextObject>(ToolsContext->ToolManager);
+		if (ensure(GizmoContextObject))
+		{
+			GizmoContextObject->RegisterGizmosWithManager(ToolsContext->ToolManager);
+			return true;
+		}
+	}
+	return false;
+}
+
+bool UE::SkeletalMeshGizmoUtils::UnregisterTransformGizmoContextObject(UInteractiveToolsContext* ToolsContext)
+{
+	if (ensure(ToolsContext))
+	{
+		USkeletalMeshGizmoContextObject* Found = ToolsContext->ContextObjectStore->FindContext<USkeletalMeshGizmoContextObject>();
+		if (Found != nullptr)
+		{
+			Found->UnregisterGizmosWithManager(ToolsContext->ToolManager);
+			ToolsContext->ContextObjectStore->RemoveContextObject(Found);
+		}
+		return true;
+	}
+	return false;
+}
+
+UTransformGizmo* UE::SkeletalMeshGizmoUtils::CreateTransformGizmo(UInteractiveToolManager* ToolManager, void* Owner)
+{
+	if (!ToolManager)
+	{
+		return nullptr;
+	}
+	
+	UEditorInteractiveGizmoManager* GizmoManager = Cast<UEditorInteractiveGizmoManager>(ToolManager->GetPairedGizmoManager());
+	if (!GizmoManager)
+	{
+		return nullptr;
+	}
+
+	const USkeletalMeshGizmoContextObject* Found = GizmoManager->GetContextObjectStore()->FindContext<USkeletalMeshGizmoContextObject>();
+	if (!Found)
+	{
+		return nullptr;
+	}
+
+	static const FString InstanceIdentifier;
+	UTransformGizmo* NewGizmo = Cast<UTransformGizmo>( GizmoManager->CreateGizmo(
+		USkeletalMeshGizmoContextObject::TransformBuilderIdentifier(), InstanceIdentifier, Owner) );
+
+	if (ensure(NewGizmo))
+	{
+		GizmoManager->OnGizmosParametersChangedDelegate().AddUObject(NewGizmo, &UTransformGizmo::OnParametersChanged);
+		if (GizmoManager->GetDefaultGizmosParameters())
+		{
+			NewGizmo->OnParametersChanged(*GizmoManager->GetDefaultGizmosParameters());
+		}
+	}
+	
+	return NewGizmo;
+}
+
+const FString& USkeletalMeshGizmoContextObject::TransformBuilderIdentifier()
+{
+	static const FString Identifier(TEXT("SkeletalMeshTransformGizmo"));
+	return Identifier;
+}
+
+void USkeletalMeshGizmoContextObject::RegisterGizmosWithManager(UInteractiveToolManager* InToolManager)
+{
+	if (ensure(!bGizmosRegistered) == false)
+	{
+		return;
+	}
+
+	InToolManager->GetContextObjectStore()->AddContextObject(this);
+	
+	UEditorInteractiveGizmoManager* GizmoManager = Cast<UEditorInteractiveGizmoManager>(InToolManager->GetPairedGizmoManager());
+	if (ensure(GizmoManager))
+	{
+		UEditorTransformGizmoBuilder* GizmoBuilder = NewObject<UEditorTransformGizmoBuilder>(this);
+		GizmoBuilder->CustomizationFunction = [] { return GetGizmoCustomization(); };
+		GizmoManager->RegisterGizmoType(TransformBuilderIdentifier(), GizmoBuilder);
+	}
+	
+	bGizmosRegistered = true;
+}
+
+void USkeletalMeshGizmoContextObject::UnregisterGizmosWithManager(UInteractiveToolManager* InToolManager)
+{
+	ensure(bGizmosRegistered);
+	
+	InToolManager->GetContextObjectStore()->RemoveContextObject(this);
+
+	UInteractiveGizmoManager* GizmoManager = InToolManager->GetPairedGizmoManager();
+	GizmoManager->DeregisterGizmoType(TransformBuilderIdentifier());
+	
+	bGizmosRegistered = false;
+}
+
+USkeletalMeshGizmoWrapperBase* USkeletalMeshGizmoContextObject::GetNewWrapper(UInteractiveToolManager* InToolManager, UObject* Outer, IGizmoStateTarget* InStateTarget)
+{
+	if (!bGizmosRegistered)
+	{
+		return nullptr;
+	}
+	
+	UTransformGizmo* Gizmo = UE::SkeletalMeshGizmoUtils::CreateTransformGizmo(InToolManager, Outer);
+	if (!Gizmo)
+	{
+		return nullptr;
+	}
+
+	Gizmo->TransformGizmoSource = nullptr;
+	Gizmo->SetVisibility(true);
+
+	USkeletonTransformProxy* SkeletonTransformProxy = NewObject<USkeletonTransformProxy>(Outer);
+	SkeletonTransformProxy->bRotatePerObject = true;
+	
+	Gizmo->SetActiveTarget(SkeletonTransformProxy, nullptr, InStateTarget);
+
+	if (UGizmoElementHitMultiTarget* HitMultiTarget = Cast< UGizmoElementHitMultiTarget>(Gizmo->HitTarget))
+	{
+		HitMultiTarget->GizmoTransformProxy = SkeletonTransformProxy;
+	}
+
+	USkeletalMeshGizmoWrapper* GizmoWrapper = NewObject<USkeletalMeshGizmoWrapper>(Outer);
+	GizmoWrapper->TransformGizmo = Gizmo;
+	GizmoWrapper->TransformProxy = SkeletonTransformProxy;
+
+	if (Gizmo->HitTarget)
+	{
+		TWeakObjectPtr<USkeletalMeshGizmoWrapper> WeakWrapper = MakeWeakObjectPtr(GizmoWrapper);
+		Gizmo->HitTarget->Condition = [WeakWrapper](const FInputDeviceRay&)
+		{
+			return WeakWrapper.IsValid() ? WeakWrapper->CanInteract() : false;
+		};
+	}
+	
+	return GizmoWrapper;
+}
+
+const FGizmoCustomization& USkeletalMeshGizmoContextObject::GetGizmoCustomization()
+{
+	static const FString MaterialName = TEXT("/SkeletalMeshModelingTools/SkeletalMeshGizmoMaterial.SkeletalMeshGizmoMaterial");
+	UMaterial* Material = LoadObject<UMaterial>(nullptr, *MaterialName, nullptr, LOAD_None, nullptr);
+
+	check(Material);
+	
+	static FGizmoCustomization Customization;
+		Customization.Material = Material;
+		Customization.SizeCoefficient = 1.5f;
+
+	return Customization;
+}
+
+bool USkeletalMeshGizmoWrapper::CanInteract() const
+{
+	if (!TransformProxy || !TransformGizmo)
+	{
+		return false;
+	}
+	return TransformProxy->IsValid() && TransformGizmo->bVisible; 
+}
+
+bool USkeletalMeshGizmoWrapper::IsGizmoHit(const FInputDeviceRay& PressPos) const
+{
+	if (!TransformGizmo || !TransformGizmo->bVisible)
+	{
+		return false;
+	}
+	return TransformGizmo->HitTarget ? TransformGizmo->HitTarget->IsHit(PressPos).bHit : false;
+}
+
+void USkeletalMeshGizmoWrapper::Clear()
+{
+	if (TransformProxy)
+	{
+		TransformProxy->Initialize(FTransform::Identity, EToolContextCoordinateSystem::Local);
+	}
+	
+	if (TransformGizmo)
+	{
+		TransformGizmo->SetVisibility(false);
+		TransformGizmo->ClearActiveTarget();
+
+		if (UInteractiveGizmoManager* GizmoManager = TransformGizmo->GetGizmoManager())
+		{
+			GizmoManager->DestroyGizmo(TransformGizmo);
+		}
+	}
+}
+
+void USkeletalMeshGizmoWrapper::Initialize(const FTransform& InTransform, const EToolContextCoordinateSystem& InTransformMode)
+{
+	if (TransformProxy)
+	{
+		TransformProxy->Initialize(InTransform, InTransformMode);
+	}
+	
+	if (TransformGizmo)
+	{
+		TransformGizmo->SetVisibility(TransformProxy ? TransformProxy->IsValid() : false);
+	}
+}
+
+void USkeletalMeshGizmoWrapper::HandleBoneTransform(FGetTransform GetTransformFunc, FSetTransform SetTransformFunc)
+{
+	if (!TransformProxy || !Component.IsValid())
+	{
+		return;
+	}
+
+	static constexpr bool bModifyComponentOnTransform = false;
+	TransformProxy->AddComponentCustom(Component.Get(),
+		MoveTemp(GetTransformFunc),
+		MoveTemp(SetTransformFunc),
+		INDEX_NONE,
+		bModifyComponentOnTransform);
+
+	if (TransformGizmo)
+	{
+		TransformGizmo->SetVisibility(TransformProxy->IsValid());
+	}
+}
\ No newline at end of file
diff --git a/Engine/Plugins/Animation/SkeletalMeshModelingTools/Source/SkeletalMeshModelingTools/Private/SkeletalMeshGizmoUtils.h b/Engine/Plugins/Animation/SkeletalMeshModelingTools/Source/SkeletalMeshModelingTools/Private/SkeletalMeshGizmoUtils.h
new file mode 100644
--- /dev/null
+++ b/Engine/Plugins/Animation/SkeletalMeshModelingTools/Source/SkeletalMeshModelingTools/Private/SkeletalMeshGizmoUtils.h
@@ -0,0 +1,96 @@
+﻿// Copyright Epic Games, Inc. All Rights Reserved.
+
+#pragma once
+
+#include "UObject/Object.h"
+#include "Containers/UnrealString.h"
+#include "SkeletalMesh/SkeletalMeshEditionInterface.h"
+
+#include "SkeletalMeshGizmoUtils.generated.h"
+
+class UInteractiveToolsContext;
+class UInteractiveToolManager;
+class UTransformGizmo;
+class USkeletonTransformProxy;
+struct FGizmoCustomization;
+
+namespace UE
+{
+	
+namespace SkeletalMeshGizmoUtils
+{
+	/**
+	 * The functions below are helper functions that simplify usage of a USkeletalMeshGizmoContextObject
+	 * that is registered as a ContextStoreObject in an InteractiveToolsContext
+	 */
+
+	/**
+	 * If one does not already exist, create a new instance of USkeletalMeshGizmoContextObject and add it to the
+	 * ToolsContext's ContextObjectStore
+	 * @return true if the ContextObjectStore now has a USkeletalMeshGizmoContextObject (whether it already existed, or was created)
+	 */
+	SKELETALMESHMODELINGTOOLS_API bool RegisterTransformGizmoContextObject(UInteractiveToolsContext* InToolsContext);
+
+	/**
+	 * Remove any existing USkeletalMeshGizmoContextObject from the ToolsContext's ContextObjectStore
+	 * @return true if the ContextObjectStore no longer has a USkeletalMeshGizmoContextObject (whether it was removed, or did not exist)
+	 */
+	SKELETALMESHMODELINGTOOLS_API bool UnregisterTransformGizmoContextObject(UInteractiveToolsContext* InToolsContext);
+
+	/**
+	 * Spawn an editor like Transform Gizmo. ToolManager's ToolsContext must have a USkeletalMeshGizmoContextObject registered 
+	 */
+	SKELETALMESHMODELINGTOOLS_API UTransformGizmo* CreateTransformGizmo(UInteractiveToolManager* ToolManager, void* Owner = nullptr);
+}
+	
+}
+
+/**
+ * USkeletalMeshGizmoWrapper is a wrapper class to handle a Transform Gizmo and it's Transform proxy so that it can
+ * be used to update skeletal mesh infos using a Gizmo.
+ */
+
+UCLASS()
+class SKELETALMESHMODELINGTOOLS_API USkeletalMeshGizmoWrapper : public USkeletalMeshGizmoWrapperBase
+{
+	GENERATED_BODY()
+public:
+	virtual void Initialize(const FTransform& InTransform = FTransform::Identity, const EToolContextCoordinateSystem& InTransformMode = EToolContextCoordinateSystem::Local) override;
+	virtual void HandleBoneTransform(FGetTransform GetTransformFunc, FSetTransform SetTransformFunc) override;
+	virtual void Clear() override;
+
+	virtual bool CanInteract() const override;
+	virtual bool IsGizmoHit(const FInputDeviceRay& PressPos) const override;
+
+	UPROPERTY()
+	TObjectPtr<UTransformGizmo> TransformGizmo;
+	UPROPERTY()
+	TObjectPtr<USkeletonTransformProxy> TransformProxy;
+};
+
+/**
+ * USkeletalMeshGizmoContextObject is a utility object that registers a Gizmo Builder for UTransformGizmo.
+ * (see UCombinedTransformGizmoContextObject for more details)
+ */
+
+UCLASS()
+class SKELETALMESHMODELINGTOOLS_API USkeletalMeshGizmoContextObject : public USkeletalMeshGizmoContextObjectBase
+{
+	GENERATED_BODY()
+
+public:
+
+	// builder identifiers for transform gizmo
+	static const FString& TransformBuilderIdentifier();
+	
+	void RegisterGizmosWithManager(UInteractiveToolManager* InToolManager);
+	void UnregisterGizmosWithManager(UInteractiveToolManager* InToolManager);
+
+	virtual USkeletalMeshGizmoWrapperBase* GetNewWrapper(UInteractiveToolManager* InToolManager, UObject* Outer = nullptr, IGizmoStateTarget* InStateTarget = nullptr) override;
+
+private:
+	
+	static const FGizmoCustomization& GetGizmoCustomization();
+	
+	bool bGizmosRegistered = false;
+};
diff --git a/Engine/Plugins/Animation/SkeletalMeshModelingTools/Source/SkeletalMeshModelingTools/Private/SkeletalMeshModelingToolsCommands.cpp b/Engine/Plugins/Animation/SkeletalMeshModelingTools/Source/SkeletalMeshModelingTools/Private/SkeletalMeshModelingToolsCommands.cpp
new file mode 100644
--- /dev/null
+++ b/Engine/Plugins/Animation/SkeletalMeshModelingTools/Source/SkeletalMeshModelingTools/Private/SkeletalMeshModelingToolsCommands.cpp
@@ -0,0 +1,72 @@
+// Copyright Epic Games, Inc. All Rights Reserved.
+
+#include "SkeletalMeshModelingToolsCommands.h"
+
+#include "SkeletalMesh/SkeletonEditingTool.h"
+
+#define LOCTEXT_NAMESPACE "SkeletalMeshModelingToolsCommands"
+
+
+void FSkeletalMeshModelingToolsCommands::RegisterCommands()
+{
+	UI_COMMAND(ToggleEditingToolsMode, "Enable Editing Tools", "Toggles editing tools on or off.", EUserInterfaceActionType::ToggleButton, FInputChord());
+	
+	UI_COMMAND(NewBone, "New Bone", "Add new Bone.", EUserInterfaceActionType::Button, FInputChord(EKeys::N, EModifierKey::Control));
+	UI_COMMAND(RemoveBone, "Remove Bone", "Remove selected bone(s).", EUserInterfaceActionType::Button, FInputChord(EKeys::Delete), FInputChord(EKeys::BackSpace));
+	UI_COMMAND(UnParentBone, "Unparent Bone", "Unparent selected bone(s).", EUserInterfaceActionType::Button, FInputChord(EKeys::P, EModifierKey::Shift));
+	UI_COMMAND(RenameBone, "Rename Bone", "Rename the selected bone.", EUserInterfaceActionType::Button, FInputChord(EKeys::F2));
+
+	UI_COMMAND(CopyBones, "Copy Bone(s)", "Copy selected bone(s).", EUserInterfaceActionType::Button, FInputChord(EKeys::C, EModifierKey::Control));
+	UI_COMMAND(PasteBones, "Paste Bone(s)", "Paste selected bone(s).", EUserInterfaceActionType::Button, FInputChord(EKeys::V, EModifierKey::Control));
+	UI_COMMAND(DuplicateBones, "Duplicate Bone(s)", "Duplicate selected bone(s).", EUserInterfaceActionType::Button, FInputChord(EKeys::D, EModifierKey::Control));
+}
+
+const FSkeletalMeshModelingToolsCommands& FSkeletalMeshModelingToolsCommands::Get()
+{
+	return TCommands<FSkeletalMeshModelingToolsCommands>::Get();
+}
+
+FSkeletalMeshModelingToolsActionCommands::FSkeletalMeshModelingToolsActionCommands() : 
+	TInteractiveToolCommands<FSkeletalMeshModelingToolsActionCommands>(
+		"SeletalMeshModelingToolsEditMode", // Context name for fast lookup
+		NSLOCTEXT("Contexts", "SeletalMeshModelingToolsEditMode", "Skeletal Mesh Modeling Tools - Shared Shortcuts"), // Localized context name for displaying
+		NAME_None, // Parent
+		FAppStyle::GetAppStyleSetName() // Icon Style Set
+	)
+{
+}
+
+void FSkeletalMeshModelingToolsActionCommands::GetToolDefaultObjectList(TArray<UInteractiveTool*>& ToolCDOs)
+{}
+
+void FSkeletalMeshModelingToolsActionCommands::RegisterAllToolActions()
+{
+	FSkeletonEditingToolActionCommands::Register();
+}
+
+void FSkeletalMeshModelingToolsActionCommands::UnregisterAllToolActions()
+{
+	FSkeletonEditingToolActionCommands::Unregister();
+}
+
+void FSkeletalMeshModelingToolsActionCommands::UpdateToolCommandBinding(UInteractiveTool* Tool, TSharedPtr<FUICommandList> UICommandList, bool bUnbind)
+{
+#define UPDATE_BINDING(CommandsType)  if (!bUnbind) CommandsType::Get().BindCommandsForCurrentTool(UICommandList, Tool); else CommandsType::Get().UnbindActiveCommands(UICommandList);
+
+	if (ExactCast<USkeletonEditingTool>(Tool) != nullptr)
+	{
+		UPDATE_BINDING(FSkeletonEditingToolActionCommands);
+	}
+}
+
+#define DEFINE_TOOL_ACTION_COMMANDS(CommandsClassName, ContextNameString, SettingsDialogString, ToolClassName ) \
+CommandsClassName::CommandsClassName() : TInteractiveToolCommands<CommandsClassName>( \
+ContextNameString, NSLOCTEXT("Contexts", ContextNameString, SettingsDialogString), NAME_None, FAppStyle::GetAppStyleSetName()) {} \
+void CommandsClassName::GetToolDefaultObjectList(TArray<UInteractiveTool*>& ToolCDOs) \
+{\
+ToolCDOs.Add(GetMutableDefault<ToolClassName>()); \
+}
+
+DEFINE_TOOL_ACTION_COMMANDS(FSkeletonEditingToolActionCommands, "SkeletalMeshModelingToolsSkeletonEditing", "Skeletal Mesh Modeling Tools - Skeleton Editing Tool", USkeletonEditingTool);
+
+#undef LOCTEXT_NAMESPACE
diff --git a/Engine/Plugins/Animation/SkeletalMeshModelingTools/Source/SkeletalMeshModelingTools/Private/SkeletalMeshModelingToolsCommands.h b/Engine/Plugins/Animation/SkeletalMeshModelingTools/Source/SkeletalMeshModelingTools/Private/SkeletalMeshModelingToolsCommands.h
new file mode 100644
--- /dev/null
+++ b/Engine/Plugins/Animation/SkeletalMeshModelingTools/Source/SkeletalMeshModelingTools/Private/SkeletalMeshModelingToolsCommands.h
@@ -0,0 +1,62 @@
+// Copyright Epic Games, Inc. All Rights Reserved.
+
+#pragma once
+
+#include "CoreMinimal.h"
+#include "Styling/AppStyle.h"
+#include "Framework/Commands/Commands.h"
+#include "Tools/InteractiveToolsCommands.h"
+
+
+class FSkeletalMeshModelingToolsCommands : public TCommands<FSkeletalMeshModelingToolsCommands>
+{
+public:
+	FSkeletalMeshModelingToolsCommands()
+	    : TCommands<FSkeletalMeshModelingToolsCommands>(
+	    	TEXT("SkeletalMeshModelingTools"),
+	    	NSLOCTEXT("Contexts", "SkeletalMeshModelingToolsCommands", "Skeletal Mesh Modeling Tools"),
+	    	NAME_None,
+	    	FAppStyle::GetAppStyleSetName())
+	{
+	}
+
+	virtual void RegisterCommands() override;
+	static const FSkeletalMeshModelingToolsCommands& Get();
+
+	// Editing tools commands
+	TSharedPtr<FUICommandInfo> ToggleEditingToolsMode;
+
+	// Skeleton Editing commands
+	TSharedPtr<FUICommandInfo> NewBone;
+	TSharedPtr<FUICommandInfo> RemoveBone;
+	TSharedPtr<FUICommandInfo> UnParentBone;
+	TSharedPtr<FUICommandInfo> RenameBone;
+
+	// Copy & Paste commands
+	TSharedPtr<FUICommandInfo> CopyBones;
+	TSharedPtr<FUICommandInfo> PasteBones;
+	TSharedPtr<FUICommandInfo> DuplicateBones;
+};
+
+class FSkeletalMeshModelingToolsActionCommands : public TInteractiveToolCommands<FSkeletalMeshModelingToolsActionCommands>
+{
+public:
+	FSkeletalMeshModelingToolsActionCommands();
+
+	// TInteractiveToolCommands
+	virtual void GetToolDefaultObjectList(TArray<UInteractiveTool*>& ToolCDOs) override;
+	
+	static void RegisterAllToolActions();
+	static void UnregisterAllToolActions();
+	static void UpdateToolCommandBinding(UInteractiveTool* Tool, TSharedPtr<FUICommandList> UICommandList, bool bUnbind = false);
+};
+
+#define DECLARE_TOOL_ACTION_COMMANDS(CommandsClassName) \
+class CommandsClassName : public TInteractiveToolCommands<CommandsClassName> \
+{\
+public:\
+CommandsClassName();\
+virtual void GetToolDefaultObjectList(TArray<UInteractiveTool*>& ToolCDOs) override;\
+};\
+
+DECLARE_TOOL_ACTION_COMMANDS(FSkeletonEditingToolActionCommands);
\ No newline at end of file
diff --git a/Engine/Plugins/Animation/SkeletalMeshModelingTools/Source/SkeletalMeshModelingTools/Private/SkeletalMeshModelingToolsEditorMode.cpp b/Engine/Plugins/Animation/SkeletalMeshModelingTools/Source/SkeletalMeshModelingTools/Private/SkeletalMeshModelingToolsEditorMode.cpp
new file mode 100644
--- /dev/null
+++ b/Engine/Plugins/Animation/SkeletalMeshModelingTools/Source/SkeletalMeshModelingTools/Private/SkeletalMeshModelingToolsEditorMode.cpp
@@ -0,0 +1,549 @@
+// Copyright Epic Games, Inc. All Rights Reserved.
+
+#include "SkeletalMeshModelingToolsEditorMode.h"
+
+#include "AttributeEditorTool.h"
+#include "SkeletalMeshModelingToolsEditorModeToolkit.h"
+#include "SkeletalMeshModelingToolsCommands.h"
+#include "SkeletalMeshGizmoUtils.h"
+
+#include "BaseGizmos/TransformGizmoUtil.h"
+#include "EdMode.h"
+#include "EdModeInteractiveToolsContext.h"
+#include "EditorModeManager.h"
+#include "EditorViewportClient.h"
+#include "Features/IModularFeatures.h"
+
+// Stylus support is currently disabled due to issues with the stylus plugin
+// We are leaving the code in this cpp file, defined out, so that it is easier to bring back if/when the stylus plugin is improved.
+#define ENABLE_STYLUS_SUPPORT 0
+
+#if ENABLE_STYLUS_SUPPORT 
+#include "IStylusInputModule.h"
+#include "IStylusState.h"
+#endif
+
+#include "ToolTargets/SkeletalMeshComponentToolTarget.h"
+#include "Components/SkeletalMeshComponent.h"
+
+#include "ConvertToPolygonsTool.h"
+#include "DeformMeshPolygonsTool.h"
+#include "DisplaceMeshTool.h"
+#include "DynamicMeshSculptTool.h"
+#include "EditMeshPolygonsTool.h"
+#include "EditorInteractiveGizmoManager.h"
+#include "HoleFillTool.h"
+#include "ISkeletalMeshEditor.h"
+#include "LatticeDeformerTool.h"
+#include "MeshAttributePaintTool.h"
+#include "MeshGroupPaintTool.h"
+#include "MeshSpaceDeformerTool.h"
+#include "MeshVertexPaintTool.h"
+#include "MeshVertexSculptTool.h"
+#include "ModelingToolsManagerActions.h"
+#include "OffsetMeshTool.h"
+#include "PersonaModule.h"
+#include "PolygonOnMeshTool.h"
+#include "ProjectToTargetTool.h"
+#include "RemeshMeshTool.h"
+#include "RemoveOccludedTrianglesTool.h"
+#include "SimplifyMeshTool.h"
+#include "SkeletalMeshEditorUtils.h"
+#include "SmoothMeshTool.h"
+#include "ToolTargetManager.h"
+#include "WeldMeshEdgesTool.h"
+
+#include "SkeletalMeshNotifier.h"
+#include "Animation/DebugSkelMeshComponent.h"
+
+#include "SkeletalMesh/SkeletalMeshEditionInterface.h"
+#include "SkeletalMesh/SkeletonEditingTool.h"
+#include "SkeletalMesh/SkinWeightsBindingTool.h"
+#include "SkeletalMesh/SkinWeightsPaintTool.h"
+#include "ToolTargets/SkeletalMeshToolTarget.h"
+
+#define LOCTEXT_NAMESPACE "SkeletalMeshModelingToolsEditorMode"
+
+#if ENABLE_STYLUS_SUPPORT
+// FStylusStateTracker registers itself as a listener for stylus events and implements
+// the IToolStylusStateProviderAPI interface, which allows MeshSurfacePointTool implementations
+ // to query for the pen pressure.
+//
+// This is kind of a hack. Unfortunately the current Stylus module is a Plugin so it
+// cannot be used in the base ToolsFramework, and we need this in the Mode as a workaround.
+//
+class FStylusStateTracker : public IStylusMessageHandler, public IToolStylusStateProviderAPI
+{
+public:
+	const IStylusInputDevice* ActiveDevice = nullptr;
+	int32 ActiveDeviceIndex = -1;
+
+	bool bPenDown = false;
+	float ActivePressure = 1.0;
+
+	FStylusStateTracker()
+	{
+		UStylusInputSubsystem* StylusSubsystem = GEditor->GetEditorSubsystem<UStylusInputSubsystem>();
+		StylusSubsystem->AddMessageHandler(*this);
+
+		ActiveDevice = FindFirstPenDevice(StylusSubsystem, ActiveDeviceIndex);
+		bPenDown = false;
+	}
+
+	virtual ~FStylusStateTracker()
+	{
+		if (GEditor)
+		{
+			if (UStylusInputSubsystem* StylusSubsystem = GEditor->GetEditorSubsystem<UStylusInputSubsystem>())
+			{
+				StylusSubsystem->RemoveMessageHandler(*this);
+			}
+		}
+	}
+
+	void OnStylusStateChanged(const FStylusState& NewState, int32 StylusIndex) override
+	{
+		if (ActiveDevice == nullptr)
+		{
+			UStylusInputSubsystem* StylusSubsystem = GEditor->GetEditorSubsystem<UStylusInputSubsystem>();
+			ActiveDevice = FindFirstPenDevice(StylusSubsystem, ActiveDeviceIndex);
+			bPenDown = false;
+		}
+		if (ActiveDevice != nullptr && ActiveDeviceIndex == StylusIndex)
+		{
+			bPenDown = NewState.IsStylusDown();
+			ActivePressure = NewState.GetPressure();
+		}
+	}
+
+
+	bool HaveActiveStylusState() const
+	{
+		return ActiveDevice != nullptr && bPenDown;
+	}
+
+	static const IStylusInputDevice* FindFirstPenDevice(const UStylusInputSubsystem* StylusSubsystem, int32& ActiveDeviceOut)
+	{
+		int32 NumDevices = StylusSubsystem->NumInputDevices();
+		for (int32 k = 0; k < NumDevices; ++k)
+		{
+			const IStylusInputDevice* Device = StylusSubsystem->GetInputDevice(k);
+			const TArray<EStylusInputType>& Inputs = Device->GetSupportedInputs();
+			for (EStylusInputType Input : Inputs)
+			{
+				if (Input == EStylusInputType::Pressure)
+				{
+					ActiveDeviceOut = k;
+					return Device;
+				}
+			}
+		}
+		return nullptr;
+	}
+
+
+
+	// IToolStylusStateProviderAPI implementation
+	virtual float GetCurrentPressure() const override
+	{
+		return (ActiveDevice != nullptr && bPenDown) ? ActivePressure : 1.0f;
+	}
+
+};
+#endif // ENABLE_STYLUS_SUPPORT
+
+
+// NOTE: This is a simple proxy at the moment. In the future we want to pull in more of the 
+// modeling tools as we add support in the skelmesh storage.
+
+const FEditorModeID USkeletalMeshModelingToolsEditorMode::Id("SkeletalMeshModelingToolsEditorMode");
+
+
+USkeletalMeshModelingToolsEditorMode::USkeletalMeshModelingToolsEditorMode() 
+{
+	Info = FEditorModeInfo(Id, LOCTEXT("SkeletalMeshEditingMode", "Skeletal Mesh Editing"), FSlateIcon(), false);
+}
+
+
+USkeletalMeshModelingToolsEditorMode::USkeletalMeshModelingToolsEditorMode(FVTableHelper& Helper) :
+	UBaseLegacyWidgetEdMode(Helper)
+{
+	
+}
+
+
+USkeletalMeshModelingToolsEditorMode::~USkeletalMeshModelingToolsEditorMode()
+{
+	// Implemented in the CPP file so that the destructor for TUniquePtr<FStylusStateTracker> gets correctly compiled.
+}
+
+
+void USkeletalMeshModelingToolsEditorMode::Initialize()
+{
+	UBaseLegacyWidgetEdMode::Initialize();
+}
+
+
+void USkeletalMeshModelingToolsEditorMode::Enter()
+{
+	UEdMode::Enter();
+
+	UEditorInteractiveToolsContext* EditorInteractiveToolsContext = GetInteractiveToolsContext(EToolsContextScope::Editor);
+	bDeactivateOnPIEStartStateToRestore = EditorInteractiveToolsContext->GetDeactivateToolsOnPIEStart();
+	EditorInteractiveToolsContext->SetDeactivateToolsOnPIEStart(false);
+
+	UEditorInteractiveToolsContext* InteractiveToolsContext = GetInteractiveToolsContext();
+
+	if (TObjectPtr<UToolTargetManager> ToolTargetManager = InteractiveToolsContext->TargetManager)
+	{
+		ToolTargetManager->AddTargetFactory( NewObject<USkeletalMeshComponentToolTargetFactory>(ToolTargetManager) );
+		ToolTargetManager->AddTargetFactory( NewObject<USkeletalMeshReadOnlyToolTargetFactory>(ToolTargetManager) );
+	}
+
+#if ENABLE_STYLUS_SUPPORT
+	StylusStateTracker = MakeUnique<FStylusStateTracker>();
+#endif
+	// register gizmo helper
+	UE::TransformGizmoUtil::RegisterTransformGizmoContextObject(InteractiveToolsContext);
+	UE::SkeletalMeshGizmoUtils::RegisterTransformGizmoContextObject(InteractiveToolsContext);
+	UE::SkeletalMeshEditorUtils::RegisterEditorContextObject(InteractiveToolsContext);
+
+	const FModelingToolsManagerCommands& ToolManagerCommands = FModelingToolsManagerCommands::Get();
+
+	RegisterTool(ToolManagerCommands.BeginPolyEditTool, TEXT("BeginPolyEditTool"), NewObject<UEditMeshPolygonsToolBuilder>());
+	UEditMeshPolygonsToolBuilder* TriEditBuilder = NewObject<UEditMeshPolygonsToolBuilder>();
+	TriEditBuilder->bTriangleMode = true;
+	RegisterTool(ToolManagerCommands.BeginTriEditTool, TEXT("BeginTriEditTool"), TriEditBuilder);
+	RegisterTool(ToolManagerCommands.BeginPolyDeformTool, TEXT("BeginPolyDeformTool"), NewObject<UDeformMeshPolygonsToolBuilder>());
+	RegisterTool(ToolManagerCommands.BeginHoleFillTool, TEXT("BeginHoleFillTool"), NewObject<UHoleFillToolBuilder>());
+	RegisterTool(ToolManagerCommands.BeginPolygonCutTool, TEXT("BeginPolyCutTool"), NewObject<UPolygonOnMeshToolBuilder>());
+
+	RegisterTool(ToolManagerCommands.BeginSimplifyMeshTool, TEXT("BeginSimplifyMeshTool"), NewObject<USimplifyMeshToolBuilder>());
+	RegisterTool(ToolManagerCommands.BeginRemeshMeshTool, TEXT("BeginRemeshMeshTool"), NewObject<URemeshMeshToolBuilder>());
+	RegisterTool(ToolManagerCommands.BeginWeldEdgesTool, TEXT("BeginWeldEdgesTool"), NewObject<UWeldMeshEdgesToolBuilder>());
+	RegisterTool(ToolManagerCommands.BeginRemoveOccludedTrianglesTool, TEXT("BeginRemoveOccludedTrianglesTool"), NewObject<URemoveOccludedTrianglesToolBuilder>());
+	RegisterTool(ToolManagerCommands.BeginProjectToTargetTool, TEXT("BeginProjectToTargetTool"), NewObject<UProjectToTargetToolBuilder>());
+	
+	RegisterTool(ToolManagerCommands.BeginPolyGroupsTool, TEXT("BeginPolyGroupsTool"), NewObject<UConvertToPolygonsToolBuilder>());
+	UMeshGroupPaintToolBuilder* MeshGroupPaintToolBuilder = NewObject<UMeshGroupPaintToolBuilder>();
+#if ENABLE_STYLUS_SUPPORT 
+	MeshGroupPaintToolBuilder->StylusAPI = StylusStateTracker.Get();
+#endif
+	RegisterTool(ToolManagerCommands.BeginMeshGroupPaintTool, TEXT("BeginMeshGroupPaintTool"), MeshGroupPaintToolBuilder);
+	
+	UMeshVertexSculptToolBuilder* MoveVerticesToolBuilder = NewObject<UMeshVertexSculptToolBuilder>();
+#if ENABLE_STYLUS_SUPPORT
+	MoveVerticesToolBuilder->StylusAPI = StylusStateTracker.Get();
+#endif
+	RegisterTool(ToolManagerCommands.BeginSculptMeshTool, TEXT("BeginSculptMeshTool"), MoveVerticesToolBuilder);
+
+	UDynamicMeshSculptToolBuilder* DynaSculptToolBuilder = NewObject<UDynamicMeshSculptToolBuilder>();
+	DynaSculptToolBuilder->bEnableRemeshing = true;
+#if ENABLE_STYLUS_SUPPORT
+	DynaSculptToolBuilder->StylusAPI = StylusStateTracker.Get();
+#endif
+	RegisterTool(ToolManagerCommands.BeginRemeshSculptMeshTool, TEXT("BeginRemeshSculptMeshTool"), DynaSculptToolBuilder);
+	
+	RegisterTool(ToolManagerCommands.BeginSmoothMeshTool, TEXT("BeginSmoothMeshTool"), NewObject<USmoothMeshToolBuilder>());
+	RegisterTool(ToolManagerCommands.BeginOffsetMeshTool, TEXT("BeginOffsetMeshTool"), NewObject<UOffsetMeshToolBuilder>());
+	RegisterTool(ToolManagerCommands.BeginMeshSpaceDeformerTool, TEXT("BeginMeshSpaceDeformerTool"), NewObject<UMeshSpaceDeformerToolBuilder>());
+	RegisterTool(ToolManagerCommands.BeginLatticeDeformerTool, TEXT("BeginLatticeDeformerTool"), NewObject<ULatticeDeformerToolBuilder>());
+	RegisterTool(ToolManagerCommands.BeginDisplaceMeshTool, TEXT("BeginDisplaceMeshTool"), NewObject<UDisplaceMeshToolBuilder>());
+
+	RegisterTool(ToolManagerCommands.BeginAttributeEditorTool, TEXT("BeginAttributeEditorTool"), NewObject<UAttributeEditorToolBuilder>());
+	RegisterTool(ToolManagerCommands.BeginMeshAttributePaintTool, TEXT("BeginMeshAttributePaintTool"), NewObject<UMeshAttributePaintToolBuilder>());
+	RegisterTool(ToolManagerCommands.BeginMeshVertexPaintTool, TEXT("BeginMeshVertexPaintTool"), NewObject<UMeshVertexPaintToolBuilder>());
+	RegisterTool(ToolManagerCommands.BeginSkinWeightsPaintTool, TEXT("BeginSkinWeightsPaintTool"), NewObject<USkinWeightsPaintToolBuilder>());
+	RegisterTool(ToolManagerCommands.BeginSkinWeightsBindingTool, TEXT("BeginSkinWeightsBindingTool"), NewObject<USkinWeightsBindingToolBuilder>());
+
+	// Skeleton Editing
+	RegisterTool(ToolManagerCommands.BeginSkeletonEditingTool, TEXT("BeginSkeletonEditingTool"), NewObject<USkeletonEditingToolBuilder>());
+
+	// register extensions
+	RegisterExtensions();
+	
+	// highlights skin weights tool by default
+	GetInteractiveToolsContext()->ToolManager->SelectActiveToolType(EToolSide::Left, TEXT("BeginSkinWeightsPaintTool"));
+
+	// record switching behavior to restore on exit
+	ToolSwitchModeToRestoreOnExit = GetInteractiveToolsContext()->ToolManager->GetToolSwitchMode();
+	// default to NOT applying changes to skeletal meshes when switching between tools without accepting
+	GetInteractiveToolsContext()->ToolManager->SetToolSwitchMode(EToolManagerToolSwitchMode::CancelIfAble);
+}
+
+UDebugSkelMeshComponent* USkeletalMeshModelingToolsEditorMode::GetSkelMeshComponent() const
+{
+	FToolBuilderState State; GetToolManager()->GetContextQueriesAPI()->GetCurrentSelectionState(State);
+	UActorComponent* SkeletalMeshComponent = ToolBuilderUtil::FindFirstComponent(State, [&](UActorComponent* Component)
+	{
+		return IsValid(Component) && Component->IsA<UDebugSkelMeshComponent>();
+	});
+
+	return Cast<UDebugSkelMeshComponent>(SkeletalMeshComponent);
+}
+
+void USkeletalMeshModelingToolsEditorMode::Exit()
+{
+	UEditorInteractiveToolsContext* InteractiveToolsContext = GetInteractiveToolsContext();
+	UE::TransformGizmoUtil::DeregisterTransformGizmoContextObject(InteractiveToolsContext);
+	UE::SkeletalMeshGizmoUtils::UnregisterTransformGizmoContextObject(InteractiveToolsContext);
+	UE::SkeletalMeshEditorUtils::UnregisterEditorContextObject(InteractiveToolsContext);
+	
+	UEditorInteractiveToolsContext* EditorInteractiveToolsContext = GetInteractiveToolsContext(EToolsContextScope::Editor);
+	EditorInteractiveToolsContext->SetDeactivateToolsOnPIEStart(bDeactivateOnPIEStartStateToRestore);
+
+	// restore previous tool switching behavior
+	GetInteractiveToolsContext()->ToolManager->SetToolSwitchMode(ToolSwitchModeToRestoreOnExit);
+	
+#if ENABLE_STYLUS_SUPPORT
+	StylusStateTracker = nullptr;
+#endif
+
+	UEdMode::Exit();
+}
+
+void USkeletalMeshModelingToolsEditorMode::RegisterExtensions()
+{
+	TArray<ISkeletalMeshModelingModeToolExtension*> Extensions = IModularFeatures::Get().GetModularFeatureImplementations<ISkeletalMeshModelingModeToolExtension>(
+		ISkeletalMeshModelingModeToolExtension::GetModularFeatureName());
+	if (Extensions.IsEmpty())
+	{
+		return;
+	}
+
+	UEditorInteractiveToolsContext* ToolsContext = GetInteractiveToolsContext();
+	
+	FExtensionToolQueryInfo ExtensionQueryInfo;
+	ExtensionQueryInfo.ToolsContext = ToolsContext;
+	ExtensionQueryInfo.AssetAPI = nullptr;
+#if ENABLE_STYLUS_SUPPORT
+	ExtensionQueryInfo.StylusAPI = StylusStateTracker.Get();
+#endif
+
+	for (ISkeletalMeshModelingModeToolExtension* Extension: Extensions)
+	{
+		TArray<FExtensionToolDescription> ToolSet;
+		Extension->GetExtensionTools(ExtensionQueryInfo, ToolSet);
+		for (const FExtensionToolDescription& ToolInfo : ToolSet)
+		{
+			RegisterTool(ToolInfo.ToolCommand, ToolInfo.ToolName.ToString(), ToolInfo.ToolBuilder);
+			ExtensionToolToInfo.Add(ToolInfo.ToolName.ToString(), ToolInfo);
+		}
+
+		TArray<TSubclassOf<UToolTargetFactory>> ExtensionToolTargetFactoryClasses;
+		if (Extension->GetExtensionToolTargets(ExtensionToolTargetFactoryClasses))
+		{
+			for (const TSubclassOf<UToolTargetFactory>& ExtensionTargetFactoryClass : ExtensionToolTargetFactoryClasses)
+			{
+				ToolsContext->TargetManager->AddTargetFactory(NewObject<UToolTargetFactory>(GetToolManager(), ExtensionTargetFactoryClass.Get()));
+			}
+		}
+	}
+}
+
+bool USkeletalMeshModelingToolsEditorMode::TryGetExtensionToolCommandGetter(
+	UInteractiveToolManager* InManager, const UInteractiveTool* InTool, 
+	TFunction<const UE::IInteractiveToolCommandsInterface& ()>& GetterOut) const
+{
+	if (!ensure(InManager && InTool)
+		|| InManager->GetActiveTool(EToolSide::Mouse) != InTool)
+	{
+		return false;
+	}
+
+	const FString ToolName = InManager->GetActiveToolName(EToolSide::Mouse);
+	if (ToolName.IsEmpty())
+	{
+		return false;
+	}
+
+	const FExtensionToolDescription* ToolDescription = ExtensionToolToInfo.Find(ToolName);
+	if (!ToolDescription || !ToolDescription->ToolCommandsGetter)
+	{
+		return false;
+	}
+	
+	GetterOut = ToolDescription->ToolCommandsGetter;
+	return true;
+}
+
+void USkeletalMeshModelingToolsEditorMode::CreateToolkit()
+{
+	Toolkit = MakeShareable(new FSkeletalMeshModelingToolsEditorModeToolkit);
+}
+
+bool USkeletalMeshModelingToolsEditorMode::IsCompatibleWith(FEditorModeID OtherModeID) const
+{
+	if (OtherModeID == FPersonaEditModes::SkeletonSelection)
+	{
+		return true;
+	}
+	
+	return Super::IsCompatibleWith(OtherModeID);
+}
+
+void USkeletalMeshModelingToolsEditorMode::Tick(FEditorViewportClient* ViewportClient, float DeltaTime)
+{
+	Super::Tick(ViewportClient, DeltaTime);
+}
+
+bool USkeletalMeshModelingToolsEditorMode::HandleClick(FEditorViewportClient* InViewportClient, HHitProxy *HitProxy, const FViewportClick &Click)
+{
+	if (Binding.IsValid())
+	{
+		TSharedPtr<ISkeletalMeshEditorBinding> BindingPtr = Binding.Pin();
+
+		TArray<FName> SelectedBones;
+		if (HitProxy && BindingPtr->GetNameFunction())
+		{
+			if (TOptional<FName> BoneName = BindingPtr->GetNameFunction()(HitProxy))
+			{
+				SelectedBones.Emplace(*BoneName);
+			}
+		}
+		
+		BindingPtr->GetNotifier().HandleNotification( SelectedBones, ESkeletalMeshNotifyType::BonesSelected);
+	}
+	
+	return Super::HandleClick(InViewportClient, HitProxy, Click);
+}
+
+bool USkeletalMeshModelingToolsEditorMode::ComputeBoundingBoxForViewportFocus(AActor* Actor, UPrimitiveComponent* PrimitiveComponent, FBox& InOutBox) const
+{
+	// if Tool supports custom Focus box, use that first
+	if (GetToolManager()->HasAnyActiveTool())
+	{
+		UInteractiveTool* Tool = GetToolManager()->GetActiveTool(EToolSide::Mouse);
+		IInteractiveToolCameraFocusAPI* FocusAPI = Cast<IInteractiveToolCameraFocusAPI>(Tool);
+		if (FocusAPI && FocusAPI->SupportsWorldSpaceFocusBox())
+		{
+			InOutBox = FocusAPI->GetWorldSpaceFocusBox();
+			return true;
+		}
+	}
+
+	// focus using selected bones in skel mesh editor
+	if (const USkeletalMeshComponent* Component = Cast<USkeletalMeshComponent>(PrimitiveComponent))
+	{
+		check(Component->GetSkeletalMeshAsset());
+		
+		if (Binding.IsValid())
+		{
+			const TArray<FName> Selection = Binding.Pin()->GetSelectedBones();
+			if (!Selection.IsEmpty())
+			{
+				TArray<FName> AllChildren;
+
+				const FReferenceSkeleton& RefSkeleton = Component->GetSkeletalMeshAsset()->GetRefSkeleton();
+				for (const FName& BoneName: Selection)
+				{
+					const int32 BoneIndex = RefSkeleton.FindBoneIndex(BoneName);
+					if (BoneIndex > INDEX_NONE)
+					{
+						// enlarge box
+						InOutBox += Component->GetBoneLocation(BoneName);
+
+						// get direct children
+						TArray<int32> Children;
+						RefSkeleton.GetDirectChildBones(BoneIndex, Children);
+						Algo::Transform(Children, AllChildren, [&RefSkeleton](int ChildrenIndex)
+						{
+							return RefSkeleton.GetBoneName(ChildrenIndex);
+						});
+					}
+				}
+
+				// enlarge box using direct children
+				for (const FName& BoneName: AllChildren)
+				{
+					InOutBox += Component->GetBoneLocation(BoneName);
+				}
+				
+				return true; 
+			}
+		}
+	}
+	
+	return Super::ComputeBoundingBoxForViewportFocus(Actor, PrimitiveComponent, InOutBox);
+}
+
+void USkeletalMeshModelingToolsEditorMode::OnToolStarted(UInteractiveToolManager* Manager, UInteractiveTool* Tool)
+{
+	FSkeletalMeshModelingToolsActionCommands::UpdateToolCommandBinding(Tool, Toolkit->GetToolkitCommands(), false);
+	if (TryGetExtensionToolCommandGetter(Manager, Tool, ExtensionToolCommandsGetter) && ensure(ExtensionToolCommandsGetter))
+	{
+		ExtensionToolCommandsGetter().BindCommandsForCurrentTool(Toolkit->GetToolkitCommands(), Tool);
+	}
+
+	// deactivate SkeletonSelection when a tool is activated.
+	// each tool is responsible for activating SkeletonSelection if necessary
+	if (Owner)
+	{
+		Owner->DeactivateMode(FPersonaEditModes::SkeletonSelection);
+	}
+}
+
+void USkeletalMeshModelingToolsEditorMode::OnToolEnded(UInteractiveToolManager* Manager, UInteractiveTool* Tool)
+{
+	FSkeletalMeshModelingToolsActionCommands::UpdateToolCommandBinding(Tool, Toolkit->GetToolkitCommands(), true);
+	if (ExtensionToolCommandsGetter)
+	{
+		ExtensionToolCommandsGetter().UnbindActiveCommands(Toolkit->GetToolkitCommands());
+		ExtensionToolCommandsGetter = nullptr;
+	}
+	
+	// reactivate SkeletonSelection when deactivating a tool
+	if (Owner)
+	{
+		Owner->ActivateMode(FPersonaEditModes::SkeletonSelection);
+	}
+}
+
+bool USkeletalMeshModelingToolsEditorMode::ShouldToolStartBeAllowed(const FString& ToolIdentifier) const
+{
+	// in the base mode, this returns false if the level editor is in PIE or simulated
+	// we allow all skeletal mesh editing tools to be started while running in PIE / simulate
+	return true;
+}
+
+void USkeletalMeshModelingToolsEditorMode::SetEditorBinding(const TWeakPtr<ISkeletalMeshEditor>& InSkeletalMeshEditor)
+{
+	if (!InSkeletalMeshEditor.IsValid())
+	{
+		return;
+	}
+	
+	Binding = InSkeletalMeshEditor.Pin()->GetBinding();
+
+	if (USkeletalMeshEditorContextObject* ContextObject = UE::SkeletalMeshEditorUtils::GetEditorContextObject(GetInteractiveToolsContext()))
+	{
+		ContextObject->Init(InSkeletalMeshEditor);
+	}
+}
+
+ISkeletalMeshEditingInterface* USkeletalMeshModelingToolsEditorMode::GetSkeletonInterface(UInteractiveTool* InTool)
+{
+	if (!IsValid(InTool) || !InTool->Implements<USkeletalMeshEditingInterface>())
+	{
+		return nullptr;
+	}
+	return static_cast<ISkeletalMeshEditingInterface*>(InTool->GetInterfaceAddress(USkeletalMeshEditingInterface::StaticClass()));
+}
+
+bool USkeletalMeshModelingToolsEditorMode::NeedsTransformGizmo() const
+{
+	UInteractiveTool* Tool = GetToolManager()->GetActiveTool(EToolSide::Mouse);
+	if (const ISkeletalMeshEditingInterface* SkeletonInterface = GetSkeletonInterface(Tool))
+	{
+		return !SkeletonInterface->GetSelectedBones().IsEmpty();
+	}
+
+	if (Binding.IsValid())
+	{
+		return !Binding.Pin()->GetSelectedBones().IsEmpty();
+	}
+	
+	return false;
+}
+
+#undef LOCTEXT_NAMESPACE
diff --git a/Engine/Plugins/Animation/SkeletalMeshModelingTools/Source/SkeletalMeshModelingTools/Private/SkeletalMeshModelingToolsEditorMode.h b/Engine/Plugins/Animation/SkeletalMeshModelingTools/Source/SkeletalMeshModelingTools/Private/SkeletalMeshModelingToolsEditorMode.h
new file mode 100644
--- /dev/null
+++ b/Engine/Plugins/Animation/SkeletalMeshModelingTools/Source/SkeletalMeshModelingTools/Private/SkeletalMeshModelingToolsEditorMode.h
@@ -0,0 +1,88 @@
+// Copyright Epic Games, Inc. All Rights Reserved.
+
+#pragma once
+
+#include "CoreMinimal.h"
+#include "EdMode.h"
+#include "SkeletalMeshModelingModeToolExtensions.h"
+
+#include "SkeletalMeshModelingToolsEditorMode.generated.h"
+
+//class FStylusStateTracker;
+class FSkeletalMeshModelingToolsEditorModeToolkit;
+class UEdModeInteractiveToolsContext;
+class ISkeletalMeshNotifier;
+class ISkeletalMeshEditorBinding;
+class ISkeletalMeshEditingInterface;
+class HHitProxy;
+class UDebugSkelMeshComponent;
+class ISkeletalMeshEditor;
+enum class EToolManagerToolSwitchMode;
+namespace UE
+{
+	class IInteractiveToolCommandsInterface;
+}
+
+UCLASS()
+class USkeletalMeshModelingToolsEditorMode : 
+	public UBaseLegacyWidgetEdMode
+{
+	GENERATED_BODY()
+public:
+	const static FEditorModeID Id;	
+
+	USkeletalMeshModelingToolsEditorMode();
+	explicit USkeletalMeshModelingToolsEditorMode(FVTableHelper& Helper);
+	virtual ~USkeletalMeshModelingToolsEditorMode() override;
+
+	// UEdMode overrides
+	virtual void Initialize() override;
+	virtual void Enter() override;
+	virtual void Exit() override;
+	virtual void CreateToolkit() override;
+	virtual bool IsCompatibleWith(FEditorModeID OtherModeID) const override;
+	virtual bool UsesTransformWidget() const override { return false; };
+	virtual bool UsesPropertyWidgets() const override { return false; };
+	virtual void Tick(FEditorViewportClient* InViewportClient, float InDeltaTime) override;
+	virtual bool HandleClick(FEditorViewportClient* InViewportClient, HHitProxy *HitProxy, const FViewportClick &Click) override;
+	virtual bool ComputeBoundingBoxForViewportFocus(AActor* Actor, UPrimitiveComponent* PrimitiveComponent, FBox& InOutBox) const override;
+	virtual bool UsesToolkits() const override { return true; }
+	virtual bool ShouldToolStartBeAllowed(const FString& ToolIdentifier) const override;
+
+	// binding
+	void SetEditorBinding(const TWeakPtr<ISkeletalMeshEditor>& InSkeletalMeshEditor);
+		
+protected:
+	virtual void OnToolStarted(UInteractiveToolManager* Manager, UInteractiveTool* Tool) override;
+	virtual void OnToolEnded(UInteractiveToolManager* Manager, UInteractiveTool* Tool) override;
+	
+private:
+	// Stylus support is currently disabled; this is left in for reference if/when it is brought back
+	//TUniquePtr<FStylusStateTracker> StylusStateTracker;
+
+	// we restore previous switch tool behavior when exiting this mode
+	EToolManagerToolSwitchMode ToolSwitchModeToRestoreOnExit;
+
+	static ISkeletalMeshEditingInterface* GetSkeletonInterface(UInteractiveTool* InTool);
+
+	UDebugSkelMeshComponent* GetSkelMeshComponent() const;
+
+	bool NeedsTransformGizmo() const;
+
+	bool bDeactivateOnPIEStartStateToRestore;
+
+	FDelegateHandle ToToolNotifierHandle;
+	FDelegateHandle FromToolNotifierHandle;
+	
+	TWeakPtr<ISkeletalMeshEditorBinding> Binding;
+
+	void RegisterExtensions();
+	// Support extension tools having their own hotkey classes
+	TMap<FString, FExtensionToolDescription> ExtensionToolToInfo;
+	// Note: this will only work when the given tool is active, because we get the tool identifier
+	//  out of the manager using GetActiveToolName
+	bool TryGetExtensionToolCommandGetter(UInteractiveToolManager* InManager, const UInteractiveTool* InTool, 
+		TFunction<const UE::IInteractiveToolCommandsInterface&()>& OutGetter) const;
+	// Used to unbind extension tool commands
+	TFunction<const UE::IInteractiveToolCommandsInterface& ()> ExtensionToolCommandsGetter;
+};
diff --git a/Engine/Plugins/Animation/SkeletalMeshModelingTools/Source/SkeletalMeshModelingTools/Private/SkeletalMeshModelingToolsEditorModeToolkit.cpp b/Engine/Plugins/Animation/SkeletalMeshModelingTools/Source/SkeletalMeshModelingTools/Private/SkeletalMeshModelingToolsEditorModeToolkit.cpp
new file mode 100644
--- /dev/null
+++ b/Engine/Plugins/Animation/SkeletalMeshModelingTools/Source/SkeletalMeshModelingTools/Private/SkeletalMeshModelingToolsEditorModeToolkit.cpp
@@ -0,0 +1,406 @@
+// Copyright Epic Games, Inc. All Rights Reserved.
+
+#include "SkeletalMeshModelingToolsEditorModeToolkit.h"
+
+#include "SkeletalMeshModelingToolsCommands.h"
+
+#include "EdModeInteractiveToolsContext.h"
+#include "Features/IModularFeatures.h"
+#include "Framework/MultiBox/MultiBoxBuilder.h"
+#include "IDetailsView.h"
+#include "InteractiveTool.h"
+#include "InteractiveToolsContext.h"
+#include "ModelingToolsEditorModeStyle.h"
+#include "Modules/ModuleManager.h"
+#include "PropertyEditorModule.h"
+#include "Tools/UEdMode.h"
+#include "Widgets/Input/SButton.h"
+#include "Widgets/Layout/SBorder.h"
+#include "Widgets/Images/SImage.h"
+#include "Widgets/Text/STextBlock.h"
+#include "ModelingToolsManagerActions.h"
+#include "SkeletalMeshModelingModeToolExtensions.h"
+#include "BaseTools/MeshSurfacePointTool.h"
+#include "SPrimaryButton.h"
+#include "Toolkits/AssetEditorModeUILayer.h"
+
+
+#define LOCTEXT_NAMESPACE "SkeletalMeshModelingToolsEditorModeToolkit"
+
+FSkeletalMeshModelingToolsEditorModeToolkit::~FSkeletalMeshModelingToolsEditorModeToolkit()
+{
+	if (UEditorInteractiveToolsContext* Context = GetScriptableEditorMode()->GetInteractiveToolsContext(EToolsContextScope::EdMode))
+	{
+		Context->OnToolNotificationMessage.RemoveAll(this);
+		Context->OnToolWarningMessage.RemoveAll(this);
+	}
+}
+
+void FSkeletalMeshModelingToolsEditorModeToolkit::Init(
+	const TSharedPtr<IToolkitHost>& InToolkitHost, 
+	TWeakObjectPtr<UEdMode> InOwningMode)
+{
+	bUsesToolkitBuilder = true;
+
+	FModeToolkit::Init(InToolkitHost, InOwningMode);
+
+	ModeWarningArea = SNew(STextBlock)
+		.AutoWrapText(true)
+		.Font(FCoreStyle::GetDefaultFontStyle("Bold", 9))
+		.ColorAndOpacity(FSlateColor(FLinearColor(0.9f, 0.15f, 0.15f)));
+	ModeWarningArea->SetText(FText::GetEmpty());
+	ModeWarningArea->SetVisibility(EVisibility::Collapsed);
+
+	ModeHeaderArea = SNew(STextBlock)
+		.AutoWrapText(true)
+		.Font(FCoreStyle::GetDefaultFontStyle("Bold", 12));
+	ModeHeaderArea->SetText(LOCTEXT("SelectToolLabel", "Select a Tool from the Toolbar"));
+	ModeHeaderArea->SetJustification(ETextJustify::Center);
+	
+	ToolWarningArea = SNew(STextBlock)
+		.AutoWrapText(true)
+		.Font(FCoreStyle::GetDefaultFontStyle("Bold", 9))
+		.ColorAndOpacity(FSlateColor(FLinearColor(0.9f, 0.15f, 0.15f)));
+	ToolWarningArea->SetText(FText::GetEmpty());
+
+
+	RegisterPalettes();
+
+	ClearNotification();
+	ClearWarning();
+	
+	// create the toolkit widget
+	{
+		ToolkitSections->ModeWarningArea = ModeWarningArea;
+		ToolkitSections->DetailsView = ModeDetailsView;
+		ToolkitSections->ToolWarningArea = ToolWarningArea;
+
+		SAssignNew(ToolkitWidget, SBorder)
+		.HAlign(HAlign_Fill)
+		.Padding(0)
+		.BorderImage(FAppStyle::GetBrush("ToolPanel.GroupBorder"))
+		[
+			ToolkitBuilder->GenerateWidget()->AsShared()
+		];	
+	}
+	
+	ActiveToolName = FText::GetEmpty();
+	ActiveToolMessage = FText::GetEmpty();
+
+	if (UEditorInteractiveToolsContext* Context = GetScriptableEditorMode()->GetInteractiveToolsContext(EToolsContextScope::EdMode))
+	{
+		Context->OnToolNotificationMessage.AddSP(this, &FSkeletalMeshModelingToolsEditorModeToolkit::PostNotification);
+		Context->OnToolWarningMessage.AddSP(this, &FSkeletalMeshModelingToolsEditorModeToolkit::PostWarning);
+	}
+
+	// add viewport overlay widget to accept / cancel tool
+	MakeToolAcceptCancelWidget();
+}
+
+FName FSkeletalMeshModelingToolsEditorModeToolkit::GetToolkitFName() const
+{
+	return FName("SkeletalMeshModelingToolsEditorModeToolkit");
+}
+
+FText FSkeletalMeshModelingToolsEditorModeToolkit::GetBaseToolkitName() const
+{
+	return LOCTEXT("DisplayName", "Skeletal Mesh Editing Tools");
+}
+
+void FSkeletalMeshModelingToolsEditorModeToolkit::OnToolStarted(UInteractiveToolManager* Manager, UInteractiveTool* Tool)
+{
+	UpdateActiveToolProperties(Tool);
+
+	Tool->OnPropertySetsModified.AddSP(this, &FSkeletalMeshModelingToolsEditorModeToolkit::UpdateActiveToolProperties, Tool);
+
+	ModeHeaderArea->SetVisibility(EVisibility::Collapsed);
+	ActiveToolName = Tool->GetToolInfo().ToolDisplayName;
+
+	FString ActiveToolIdentifier = GetScriptableEditorMode()->GetToolManager()->GetActiveToolName(EToolSide::Left);
+	ActiveToolIdentifier.InsertAt(0, ".");
+	FName ActiveToolIconName = ISlateStyle::Join(FModelingToolsManagerCommands::Get().GetContextName(), TCHAR_TO_ANSI(*ActiveToolIdentifier));
+	ActiveToolIcon = FModelingToolsEditorModeStyle::Get()->GetOptionalBrush(ActiveToolIconName);
+	
+	GetToolkitHost()->AddViewportOverlayWidget(ViewportOverlayWidget.ToSharedRef());
+}
+
+void FSkeletalMeshModelingToolsEditorModeToolkit::OnToolEnded(UInteractiveToolManager* Manager, UInteractiveTool* Tool)
+{
+	if (IsHosted())
+	{
+		GetToolkitHost()->RemoveViewportOverlayWidget(ViewportOverlayWidget.ToSharedRef());
+	}
+	
+	if (Tool)
+	{
+		Tool->OnPropertySetsModified.RemoveAll(this);
+	}
+
+	ModeHeaderArea->SetVisibility(EVisibility::Visible);
+	ModeDetailsView->SetObject(nullptr);
+	ActiveToolName = FText::GetEmpty();
+	ClearNotification();
+	ClearWarning();
+}
+
+FText FSkeletalMeshModelingToolsEditorModeToolkit::GetActiveToolDisplayName() const
+{
+	return ActiveToolName;
+}
+
+
+FText FSkeletalMeshModelingToolsEditorModeToolkit::GetActiveToolMessage() const
+{
+	return ActiveToolMessage;
+}
+
+void FSkeletalMeshModelingToolsEditorModeToolkit::RegisterPalettes()
+{
+	const FModelingToolsManagerCommands& Commands = FModelingToolsManagerCommands::Get();
+	const TSharedPtr<FUICommandList> CommandList = GetToolkitCommands();
+	
+	ToolkitSections = MakeShared<FToolkitSections>();
+	FToolkitBuilderArgs ToolkitBuilderArgs(GetScriptableEditorMode()->GetModeInfo().ToolbarCustomizationName);
+	ToolkitBuilderArgs.ToolkitCommandList = GetToolkitCommands();
+	ToolkitBuilderArgs.ToolkitSections = ToolkitSections;
+	ToolkitBuilderArgs.SelectedCategoryTitleVisibility = EVisibility::Collapsed;
+	ToolkitBuilder = MakeShared<FToolkitBuilder>(ToolkitBuilderArgs);
+
+	const TArray<TSharedPtr<FUICommandInfo>> SkeletonCommands({
+		Commands.BeginSkeletonEditingTool
+	});
+	ToolkitBuilder->AddPalette(
+		MakeShareable( new FToolPalette( Commands.LoadSkeletonTools.ToSharedRef(), SkeletonCommands ) ) );
+
+	const TArray<TSharedPtr<FUICommandInfo>> SkinCommands({
+		Commands.BeginSkinWeightsBindingTool,
+		Commands.BeginSkinWeightsPaintTool,
+		Commands.BeginAttributeEditorTool,
+		Commands.BeginMeshAttributePaintTool,
+		Commands.BeginMeshVertexPaintTool,
+		Commands.BeginPolyGroupsTool,
+		Commands.BeginMeshGroupPaintTool
+	});
+	ToolkitBuilder->AddPalette(
+		MakeShareable( new FToolPalette( Commands.LoadSkinTools.ToSharedRef(), SkinCommands ) ) );
+
+	const TArray<TSharedPtr<FUICommandInfo>> DeformCommands({
+		Commands.BeginSculptMeshTool,
+		Commands.BeginRemeshSculptMeshTool,
+		Commands.BeginSmoothMeshTool,
+		Commands.BeginOffsetMeshTool,
+		Commands.BeginMeshSpaceDeformerTool,
+		Commands.BeginLatticeDeformerTool,
+		Commands.BeginDisplaceMeshTool
+	});
+	ToolkitBuilder->AddPalette(
+		MakeShareable( new FToolPalette( Commands.LoadDeformTools.ToSharedRef(), DeformCommands ) ) );
+
+	const TArray<TSharedPtr<FUICommandInfo>> ModelingMeshCommands({
+		Commands.BeginPolyEditTool,
+		Commands.BeginTriEditTool,
+		Commands.BeginPolyDeformTool,
+		Commands.BeginHoleFillTool,
+		Commands.BeginPolygonCutTool,
+		});
+	ToolkitBuilder->AddPalette(
+		MakeShareable( new FToolPalette( Commands.LoadPolyTools.ToSharedRef(), ModelingMeshCommands ) ) );
+
+	const TArray<TSharedPtr<FUICommandInfo>> ProcessMeshCommands({
+		Commands.BeginSimplifyMeshTool,
+		Commands.BeginRemeshMeshTool,
+		Commands.BeginWeldEdgesTool,
+		Commands.BeginRemoveOccludedTrianglesTool,
+		Commands.BeginProjectToTargetTool
+	});
+	ToolkitBuilder->AddPalette(
+		MakeShareable( new FToolPalette( Commands.LoadMeshOpsTools.ToSharedRef(), ProcessMeshCommands ) ) );
+
+	ToolkitBuilder->SetActivePaletteOnLoad(Commands.LoadSkinTools.Get());
+	ToolkitBuilder->UpdateWidget();
+
+	RegisterExtensionsPalettes();
+	
+	// if selected palette changes, make sure we are showing the palette command buttons, which may be hidden by active Tool
+	ActivePaletteChangedHandle = ToolkitBuilder->OnActivePaletteChanged.AddLambda([this]()
+	{
+		ToolkitBuilder->SetActivePaletteCommandsVisibility(EVisibility::Visible);
+	});
+}
+
+void FSkeletalMeshModelingToolsEditorModeToolkit::RegisterExtensionsPalettes() const
+{
+	const TArray<ISkeletalMeshModelingModeToolExtension*> Extensions = IModularFeatures::Get().GetModularFeatureImplementations<ISkeletalMeshModelingModeToolExtension>(
+		ISkeletalMeshModelingModeToolExtension::GetModularFeatureName());
+
+	if (Extensions.IsEmpty())
+	{
+		return;
+	}
+
+	for (ISkeletalMeshModelingModeToolExtension* Extension: Extensions)
+	{
+		const FText ExtensionName = Extension->GetExtensionName();
+		const FText SectionName = Extension->GetToolSectionName();
+
+		FModelingModeExtensionExtendedInfo ExtensionExtendedInfo;
+		const bool bHasExtendedInfo = Extension->GetExtensionExtendedInfo(ExtensionExtendedInfo);
+
+		TSharedPtr<FUICommandInfo> PaletteCommand;
+		if (bHasExtendedInfo && ExtensionExtendedInfo.ExtensionCommand.IsValid())
+		{
+			PaletteCommand = ExtensionExtendedInfo.ExtensionCommand;
+		}
+		else
+		{
+			const FText UseTooltipText = (bHasExtendedInfo && ExtensionExtendedInfo.ToolPaletteButtonTooltip.IsEmpty() == false) ?
+				ExtensionExtendedInfo.ToolPaletteButtonTooltip : SectionName;
+			PaletteCommand = FModelingToolsManagerCommands::RegisterExtensionPaletteCommand(
+				FName(ExtensionName.ToString()),
+				SectionName, UseTooltipText, FSlateIcon());
+		}
+
+		TArray<TSharedPtr<FUICommandInfo>> PaletteItems;
+		FExtensionToolQueryInfo ExtensionQueryInfo;
+		ExtensionQueryInfo.bIsInfoQueryOnly = true;
+
+		TArray<FExtensionToolDescription> ToolSet;
+		Extension->GetExtensionTools(ExtensionQueryInfo, ToolSet);
+		
+		for (const FExtensionToolDescription& ToolInfo : ToolSet)
+		{
+			PaletteItems.Add(ToolInfo.ToolCommand);
+		}
+
+		ToolkitBuilder->AddPalette(
+			MakeShareable(new FToolPalette(PaletteCommand.ToSharedRef(), PaletteItems)));
+	}
+}
+
+void FSkeletalMeshModelingToolsEditorModeToolkit::MakeToolAcceptCancelWidget()
+{
+	SAssignNew(ViewportOverlayWidget, SHorizontalBox)
+
+	+SHorizontalBox::Slot()
+	.HAlign(HAlign_Center)
+	.VAlign(VAlign_Bottom)
+	.Padding(FMargin(0.0f, 0.0f, 0.f, 15.f))
+	[
+		SNew(SBorder)
+		.BorderImage(FAppStyle::Get().GetBrush("EditorViewport.OverlayBrush"))
+		.Padding(8.f)
+		[
+			SNew(SHorizontalBox)
+
+			+SHorizontalBox::Slot()
+			.AutoWidth()
+			.VAlign(VAlign_Center)
+			.Padding(FMargin(0.f, 0.f, 8.f, 0.f))
+			[
+				SNew(SImage)
+				.Image_Lambda([this] () { return ActiveToolIcon; })
+			]
+
+			+SHorizontalBox::Slot()
+			.AutoWidth()
+			.VAlign(VAlign_Center)
+			.Padding(FMargin(0.f, 0.f, 8.f, 0.f))
+			[
+				SNew(STextBlock)
+				.Text(this, &FSkeletalMeshModelingToolsEditorModeToolkit::GetActiveToolDisplayName)
+			]
+
+			+SHorizontalBox::Slot()
+			.AutoWidth()
+			.Padding(FMargin(0.0, 0.f, 2.f, 0.f))
+			[
+				SNew(SPrimaryButton)
+				.Text(LOCTEXT("OverlayAccept", "Accept"))
+				.ToolTipText(LOCTEXT("OverlayAcceptTooltip", "Accept/Commit the results of the active Tool [Enter]"))
+				.OnClicked_Lambda([this]() { GetScriptableEditorMode()->GetInteractiveToolsContext()->EndTool(EToolShutdownType::Accept); return FReply::Handled(); })
+				.IsEnabled_Lambda([this]() { return GetScriptableEditorMode()->GetInteractiveToolsContext()->CanAcceptActiveTool(); })
+				.Visibility_Lambda([this]() { return GetScriptableEditorMode()->GetInteractiveToolsContext()->ActiveToolHasAccept() ? EVisibility::Visible : EVisibility::Collapsed; })
+			]
+
+			+SHorizontalBox::Slot()
+			.AutoWidth()
+			.Padding(FMargin(2.0, 0.f, 0.f, 0.f))
+			[
+				SNew(SButton)
+				.TextStyle( FAppStyle::Get(), "DialogButtonText" )
+				.Text(LOCTEXT("OverlayCancel", "Cancel"))
+				.ToolTipText(LOCTEXT("OverlayCancelTooltip", "Cancel the active Tool [Esc]"))
+				.HAlign(HAlign_Center)
+				.OnClicked_Lambda([this]() { GetScriptableEditorMode()->GetInteractiveToolsContext()->EndTool(EToolShutdownType::Cancel); return FReply::Handled(); })
+				.IsEnabled_Lambda([this]() { return GetScriptableEditorMode()->GetInteractiveToolsContext()->CanCancelActiveTool(); })
+				.Visibility_Lambda([this]() { return GetScriptableEditorMode()->GetInteractiveToolsContext()->ActiveToolHasAccept() ? EVisibility::Visible : EVisibility::Collapsed; })
+			]
+
+			+SHorizontalBox::Slot()
+			.AutoWidth()
+			.Padding(FMargin(2.0, 0.f, 0.f, 0.f))
+			[
+				SNew(SPrimaryButton)
+				.Text(LOCTEXT("OverlayComplete", "Complete"))
+				.ToolTipText(LOCTEXT("OverlayCompleteTooltip", "Exit the active Tool [Enter]"))
+				.OnClicked_Lambda([this]() { GetScriptableEditorMode()->GetInteractiveToolsContext()->EndTool(EToolShutdownType::Completed); return FReply::Handled(); })
+				.IsEnabled_Lambda([this]() { return GetScriptableEditorMode()->GetInteractiveToolsContext()->CanCompleteActiveTool(); })
+				.Visibility_Lambda([this]() { return GetScriptableEditorMode()->GetInteractiveToolsContext()->CanCompleteActiveTool() ? EVisibility::Visible : EVisibility::Collapsed; })
+			]
+		]	
+	];
+}
+
+void FSkeletalMeshModelingToolsEditorModeToolkit::PostNotification(const FText& InMessage)
+{
+	ClearNotification();
+	
+	ActiveToolMessage = InMessage;
+
+	if (ModeUILayer.IsValid())
+	{
+		const FName StatusBarName = ModeUILayer.Pin()->GetStatusBarName();
+		ActiveToolMessageHandle = GEditor->GetEditorSubsystem<UStatusBarSubsystem>()->PushStatusBarMessage(StatusBarName, ActiveToolMessage);
+	}
+}
+
+
+void FSkeletalMeshModelingToolsEditorModeToolkit::ClearNotification()
+{
+	ActiveToolMessage = FText::GetEmpty();
+
+	if (ModeUILayer.IsValid())
+	{
+		const FName StatusBarName = ModeUILayer.Pin()->GetStatusBarName();
+		GEditor->GetEditorSubsystem<UStatusBarSubsystem>()->PopStatusBarMessage(StatusBarName, ActiveToolMessageHandle);
+	}
+	ActiveToolMessageHandle.Reset();
+}
+
+
+void FSkeletalMeshModelingToolsEditorModeToolkit::PostWarning(const FText& Message)
+{
+	ToolWarningArea->SetText(Message);
+	ToolWarningArea->SetVisibility(EVisibility::Visible);
+}
+
+
+void FSkeletalMeshModelingToolsEditorModeToolkit::ClearWarning()
+{
+	ToolWarningArea->SetText(FText());
+	ToolWarningArea->SetVisibility(EVisibility::Collapsed);
+}
+
+
+void FSkeletalMeshModelingToolsEditorModeToolkit::UpdateActiveToolProperties(UInteractiveTool* Tool)
+{
+	UInteractiveTool* CurTool = GetScriptableEditorMode()->GetToolManager(EToolsContextScope::EdMode)->GetActiveTool(EToolSide::Left);
+	if (CurTool == nullptr)
+	{
+		return;
+	}
+		
+	ModeDetailsView->SetObjects(CurTool->GetToolProperties(true));
+}
+
+#undef LOCTEXT_NAMESPACE
diff --git a/Engine/Plugins/Animation/SkeletalMeshModelingTools/Source/SkeletalMeshModelingTools/Private/SkeletalMeshModelingToolsEditorModeToolkit.h b/Engine/Plugins/Animation/SkeletalMeshModelingTools/Source/SkeletalMeshModelingTools/Private/SkeletalMeshModelingToolsEditorModeToolkit.h
new file mode 100644
--- /dev/null
+++ b/Engine/Plugins/Animation/SkeletalMeshModelingTools/Source/SkeletalMeshModelingTools/Private/SkeletalMeshModelingToolsEditorModeToolkit.h
@@ -0,0 +1,61 @@
+// Copyright Epic Games, Inc. All Rights Reserved.
+
+#pragma once
+
+#include "CoreMinimal.h"
+#include "StatusBarSubsystem.h"
+#include "Toolkits/BaseToolkit.h"
+
+
+class IToolkitHost;
+class STextBlock;
+class SButton;
+
+
+class FSkeletalMeshModelingToolsEditorModeToolkit : public FModeToolkit
+{
+public:
+
+	virtual ~FSkeletalMeshModelingToolsEditorModeToolkit() override;
+	
+	// IToolkit overrides
+	virtual void Init(const TSharedPtr<IToolkitHost>& InToolkitHost, TWeakObjectPtr<UEdMode> InOwningMode) override;
+	virtual FName GetToolkitFName() const override;
+	virtual FText GetBaseToolkitName() const override;
+	virtual TSharedPtr<SWidget> GetInlineContent() const override { return ToolkitWidget; }
+
+	// FModeToolkit overrides
+	virtual void OnToolStarted(UInteractiveToolManager* Manager, UInteractiveTool* Tool) override;
+	virtual void OnToolEnded(UInteractiveToolManager* Manager, UInteractiveTool* Tool) override;
+	virtual FText GetActiveToolDisplayName() const override;
+	virtual FText GetActiveToolMessage() const override;
+
+private:
+	void PostNotification(const FText& Message);
+	void ClearNotification();
+
+	void PostWarning(const FText& Message);
+	void ClearWarning();
+
+	void UpdateActiveToolProperties(UInteractiveTool* Tool);
+
+	void RegisterPalettes();
+	void RegisterExtensionsPalettes() const;
+	FDelegateHandle ActivePaletteChangedHandle;
+
+	void MakeToolAcceptCancelWidget();
+
+	FText ActiveToolName;
+	FText ActiveToolMessage;
+	FStatusBarMessageHandle ActiveToolMessageHandle;
+
+	TSharedPtr<SWidget> ToolkitWidget;
+	TSharedPtr<IDetailsView> DetailsView;
+
+	TSharedPtr<SWidget> ViewportOverlayWidget;
+	const FSlateBrush* ActiveToolIcon = nullptr;
+	
+	TSharedPtr<STextBlock> ModeWarningArea;
+	TSharedPtr<STextBlock> ModeHeaderArea;
+	TSharedPtr<STextBlock> ToolWarningArea;
+};
diff --git a/Engine/Plugins/Animation/SkeletalMeshModelingTools/Source/SkeletalMeshModelingTools/Private/SkeletalMeshModelingToolsMeshConverter.cpp b/Engine/Plugins/Animation/SkeletalMeshModelingTools/Source/SkeletalMeshModelingTools/Private/SkeletalMeshModelingToolsMeshConverter.cpp
new file mode 100644
--- /dev/null
+++ b/Engine/Plugins/Animation/SkeletalMeshModelingTools/Source/SkeletalMeshModelingTools/Private/SkeletalMeshModelingToolsMeshConverter.cpp
@@ -0,0 +1,458 @@
+﻿// Copyright Epic Games, Inc. All Rights Reserved.
+
+#include "SkeletalMeshModelingToolsMeshConverter.h"
+
+#include "AssetToolsModule.h"
+#include "Animation/Skeleton.h"
+#include "Engine/SkeletalMesh.h"
+#include "Engine/StaticMesh.h"
+#include "PropertyEditorModule.h"
+#include "StaticToSkeletalMeshConverter.h"
+#include "AssetRegistry/AssetRegistryModule.h"
+#include "Dialog/SCustomDialog.h"
+#include "Framework/Notifications/NotificationManager.h"
+#include "Widgets/Layout/SBox.h"
+#include "Widgets/Notifications/SNotificationList.h"
+
+DEFINE_LOG_CATEGORY_STATIC(LogSkeletalMeshModelingToolsMeshConverter, Log, All)
+
+#define LOCTEXT_NAMESPACE "SkeletalMeshModelingToolsMeshConverter"
+
+static void ShowEditorMessage(ELogVerbosity::Type InMessageType, const FText& InMessage, const FText* InLogMessage)
+{
+	FNotificationInfo Notification(InMessage);
+	Notification.bUseSuccessFailIcons = true;
+	Notification.FadeOutDuration = 5.0f;
+
+	SNotificationItem::ECompletionState State = SNotificationItem::CS_Success;
+
+	switch(InMessageType)
+	{
+	case ELogVerbosity::Warning:
+		UE_LOG(LogSkeletalMeshModelingToolsMeshConverter, Warning, TEXT("%s"), InLogMessage ? *InLogMessage->ToString() : *InMessage.ToString());
+		break;
+	case ELogVerbosity::Error:
+		State = SNotificationItem::CS_Fail;
+		UE_LOG(LogSkeletalMeshModelingToolsMeshConverter, Error, TEXT("%s"), InLogMessage ? *InLogMessage->ToString() : *InMessage.ToString());
+		break;
+	default:
+		checkNoEntry();
+	}
+	
+	FSlateNotificationManager::Get().AddNotification(Notification)->SetCompletionState(State);
+}
+
+
+
+USkeletonFromStaticMeshFactory::USkeletonFromStaticMeshFactory(
+	const FObjectInitializer& InObjectInitializer
+	) :
+	Super(InObjectInitializer)
+{
+	SupportedClass = USkeleton::StaticClass();
+	bCreateNew = true;
+	bEditAfterNew = false;
+}
+
+
+UObject* USkeletonFromStaticMeshFactory::FactoryCreateNew(
+	UClass* InClass,
+	UObject* InParent, 
+	FName InName,
+	EObjectFlags InFlags,
+	UObject* InContext,
+	FFeedbackContext* InWarn
+	)
+{
+	USkeleton* Skeleton = NewObject<USkeleton>(InParent, InName, InFlags);
+
+	FVector WantedRootPosition;
+	switch(PositionReference)
+	{
+	case ERootBonePositionReference::Absolute:
+		WantedRootPosition = RootPosition;
+		break;
+		
+	case ERootBonePositionReference::Relative:
+		{
+			const FBox Bounds = StaticMesh->GetBoundingBox();
+			WantedRootPosition = Bounds.Min + (Bounds.Max - Bounds.Min) * RootPosition;		
+		}
+		break;
+		
+	default:
+		checkNoEntry()
+		break;
+	}
+
+	const TCHAR* RootBoneName = TEXT("Root"); 
+	FTransform RootTransform(FTransform::Identity);
+	RootTransform.SetTranslation(WantedRootPosition);
+
+	FReferenceSkeletonModifier Modifier(Skeleton);
+	Modifier.Add(FMeshBoneInfo(RootBoneName, RootBoneName, INDEX_NONE), RootTransform);
+
+	return Skeleton;
+}
+
+
+USkeletalMeshFromStaticMeshFactory::USkeletalMeshFromStaticMeshFactory(
+	const FObjectInitializer& InObjectInitializer
+	) :
+	Super(InObjectInitializer)
+{
+	SupportedClass = USkeletalMesh::StaticClass();
+	bCreateNew = true;
+	bEditAfterNew = true;
+}
+
+
+UObject* USkeletalMeshFromStaticMeshFactory::FactoryCreateNew(
+	UClass* InClass,
+	UObject* InParent,
+	FName InName,
+	EObjectFlags InFlags,
+	UObject* InContext,
+	FFeedbackContext* InWarn
+	)
+{
+	USkeletalMesh* SkeletalMesh = NewObject<USkeletalMesh>(InParent, InName, InFlags);
+	
+	if (!FStaticToSkeletalMeshConverter::InitializeSkeletalMeshFromStaticMesh(SkeletalMesh, StaticMesh, ReferenceSkeleton, BindBoneName))
+	{
+		return nullptr;
+	}
+	
+	// Update the skeletal mesh and the skeleton so that their ref skeletons are in sync and the skeleton's preview mesh
+	// is the one we just created.
+	SkeletalMesh->SetSkeleton(Skeleton);
+	Skeleton->MergeAllBonesToBoneTree(SkeletalMesh);
+	if (!Skeleton->GetPreviewMesh())
+	{
+		Skeleton->SetPreviewMesh(SkeletalMesh);
+	}
+	
+	return SkeletalMesh;
+}
+
+
+USkeleton* UStaticMeshToSkeletalMeshConvertOptions::GetSkeleton(
+	bool& bInvalidSkeletonIsError,
+	const IPropertyHandle* PropertyHandle)
+{
+	// bInvalidSkeletonIsError = true;
+
+	switch (SkeletonImportOption)
+	{
+	default:
+	case EReferenceSkeletonImportOption::CreateNew:
+		return nullptr;
+	case EReferenceSkeletonImportOption::UseExistingSkeleton:
+		return Cast<USkeleton>(Skeleton.ResolveObject());
+	case EReferenceSkeletonImportOption::UseExistingSkeletalMesh:
+		{
+			USkeletalMesh* SkeletalMeshPtr = Cast<USkeletalMesh>(SkeletalMesh.ResolveObject());
+			return SkeletalMeshPtr ? SkeletalMeshPtr->GetSkeleton() : nullptr;
+		}
+	}
+}
+
+
+void UStaticMeshToSkeletalMeshConvertOptions::PostEditChangeProperty(FPropertyChangedEvent& PropertyChangedEvent)
+{
+	Super::PostEditChangeProperty(PropertyChangedEvent);
+
+	if (PropertyChangedEvent.Property)
+	{
+		const FName PropertyName = PropertyChangedEvent.Property->GetFName();
+		if (PropertyName == GET_MEMBER_NAME_CHECKED(UStaticMeshToSkeletalMeshConvertOptions, SkeletonImportOption) ||
+			PropertyName == GET_MEMBER_NAME_CHECKED(UStaticMeshToSkeletalMeshConvertOptions, Skeleton) ||
+			PropertyName == GET_MEMBER_NAME_CHECKED(UStaticMeshToSkeletalMeshConvertOptions, SkeletalMesh))
+		{
+			// If the pointed-at-bone doesn't exist in the new ref skeleton, then clear out the bone name.
+			if (!BindingBoneName.BoneName.IsNone())
+			{
+				const FReferenceSkeleton* RefSkeleton = nullptr;
+				if (SkeletonImportOption == EReferenceSkeletonImportOption::UseExistingSkeleton)
+				{
+					const USkeleton* SkeletonPtr = Cast<USkeleton>(Skeleton.ResolveObject());
+					RefSkeleton = SkeletonPtr ? &SkeletonPtr->GetReferenceSkeleton() : nullptr;  
+				}
+				else if (SkeletonImportOption == EReferenceSkeletonImportOption::UseExistingSkeletalMesh)
+				{
+					const USkeletalMesh* SkeletalMeshPtr = Cast<USkeletalMesh>(SkeletalMesh.ResolveObject());
+					RefSkeleton = SkeletalMeshPtr ? &SkeletalMeshPtr->GetRefSkeleton() : nullptr;
+				}
+
+				if (RefSkeleton && RefSkeleton->FindBoneIndex(BindingBoneName.BoneName) == INDEX_NONE)
+				{
+					BindingBoneName.Reset();
+				}
+			}
+		}
+	}
+}
+
+static void CreateAssetName(
+	const UPackage* InTemplatePackageName,
+	const FString& InTargetPackagePath,
+	const FString& InPrefixToRemove,
+	const FString& InPrefixToAdd,
+	const FString& InSuffixToAdd,
+	FString &OutNewAssetName,
+	FString& OutNewPackageName
+	)
+{
+	IAssetTools& AssetTools = FAssetToolsModule::GetModule().Get();
+	
+	FString TemplatePackageName = InTemplatePackageName->GetName();
+	FString TemplateAssetName = FPackageName::GetLongPackageAssetName(TemplatePackageName);
+	TemplateAssetName.RemoveFromStart(InPrefixToRemove);
+	TemplateAssetName.InsertAt(0, InPrefixToAdd);
+
+	TemplatePackageName = InTargetPackagePath + TEXT("/") + TemplateAssetName; 
+	
+	AssetTools.CreateUniqueAssetName(TemplatePackageName, InSuffixToAdd, OutNewPackageName, OutNewAssetName);
+}
+
+static bool ConvertSingleMeshToSkeletalMesh(
+	const UStaticMeshToSkeletalMeshConvertOptions* InOptions,
+	UStaticMesh* InStaticMesh,
+	TArray<UObject*>& OutObjectsAdded
+	)
+{
+	IAssetTools& AssetTools = FAssetToolsModule::GetModule().Get();
+	FReferenceSkeleton ReferenceSkeleton;
+	USkeleton* Skeleton = nullptr;
+	FName BindBoneName;
+
+	switch(InOptions->SkeletonImportOption)
+	{
+	case EReferenceSkeletonImportOption::CreateNew:
+		{
+			FString SkeletonName;
+			FString SkeletonPackageName;
+
+			CreateAssetName(InStaticMesh->GetPackage(), InOptions->DestinationPath.Path, InOptions->PrefixToRemove, InOptions->SkeletonPrefixToAdd, InOptions->SkeletonSuffixToAdd, SkeletonName, SkeletonPackageName);
+
+			USkeletonFromStaticMeshFactory* SkeletonFactory = NewObject<USkeletonFromStaticMeshFactory>();
+			SkeletonFactory->StaticMesh = InStaticMesh;
+
+			switch(InOptions->RootBonePlacement)
+			{
+			case ERootBonePlacementOptions::BottomCenter:
+				SkeletonFactory->RootPosition = FVector(0.5, 0.5, 0.0);
+				SkeletonFactory->PositionReference = ERootBonePositionReference::Relative;
+				break;
+			case ERootBonePlacementOptions::Center:
+				SkeletonFactory->RootPosition = FVector(0.5, 0.5, 0.5);
+				SkeletonFactory->PositionReference = ERootBonePositionReference::Relative;
+				break;
+			case ERootBonePlacementOptions::Origin:
+				SkeletonFactory->RootPosition = FVector(0.0, 0.0, 0.0);
+				SkeletonFactory->PositionReference = ERootBonePositionReference::Absolute;
+				break;
+			}
+
+			Skeleton = Cast<USkeleton>(AssetTools.CreateAsset(SkeletonName, FPackageName::GetLongPackagePath(SkeletonPackageName), USkeleton::StaticClass(), SkeletonFactory));
+			if (!Skeleton)
+			{
+				return false;
+			}
+			ReferenceSkeleton = Skeleton->GetReferenceSkeleton();
+			
+			OutObjectsAdded.Add(Skeleton);
+		}
+		break;
+		
+	case EReferenceSkeletonImportOption::UseExistingSkeleton:
+		{
+			Skeleton = Cast<USkeleton>(InOptions->Skeleton.ResolveObject());
+			if (!Skeleton)
+			{
+				UE_LOG(LogSkeletalMeshModelingToolsMeshConverter, Error, TEXT("No valid skeleton given as a reference. Aborting conversion."));
+				return false;
+			}
+			ReferenceSkeleton = Skeleton->GetReferenceSkeleton();
+			BindBoneName = InOptions->BindingBoneName.BoneName;
+		}
+		break;
+	case EReferenceSkeletonImportOption::UseExistingSkeletalMesh:
+		{
+			USkeletalMesh* SkeletalMesh = Cast<USkeletalMesh>(InOptions->SkeletalMesh.ResolveObject());
+			if (!SkeletalMesh)
+			{
+				UE_LOG(LogSkeletalMeshModelingToolsMeshConverter, Error, TEXT("No valid skeletal mesh given as reference. Aborting conversion."));
+				return false;
+			}
+			Skeleton = SkeletalMesh->GetSkeleton();
+			ReferenceSkeleton = SkeletalMesh->GetRefSkeleton();
+			BindBoneName = InOptions->BindingBoneName.BoneName;
+		}
+		break;
+	}
+
+	USkeletalMesh *SkeletalMesh = nullptr;
+	if (ensure(Skeleton) && ensure(ReferenceSkeleton.GetNum() != 0))
+	{
+		FString SkeletalMeshName;
+		FString SkeletalMeshPackageName;
+		
+		CreateAssetName(InStaticMesh->GetPackage(), InOptions->DestinationPath.Path, InOptions->PrefixToRemove, InOptions->SkeletalMeshPrefixToAdd, InOptions->SkeletalMeshSuffixToAdd, SkeletalMeshName, SkeletalMeshPackageName);
+
+		USkeletalMeshFromStaticMeshFactory* SkeletalMeshFactory = NewObject<USkeletalMeshFromStaticMeshFactory>();
+		SkeletalMeshFactory->StaticMesh = InStaticMesh;
+		SkeletalMeshFactory->ReferenceSkeleton = ReferenceSkeleton;
+		SkeletalMeshFactory->Skeleton = Skeleton;
+		SkeletalMeshFactory->BindBoneName = BindBoneName; 
+
+		SkeletalMesh = Cast<USkeletalMesh>(AssetTools.CreateAsset(SkeletalMeshName, FPackageName::GetLongPackagePath(SkeletalMeshPackageName), USkeletalMesh::StaticClass(), SkeletalMeshFactory));
+	}
+
+	if (ensure(SkeletalMesh))
+	{
+		OutObjectsAdded.Add(SkeletalMesh);
+		return true;
+	}
+
+	return false;
+}
+
+static bool ConvertMultipleMeshesToSkeletalMesh(
+	UStaticMeshToSkeletalMeshConvertOptions* InOptions,
+	TArray<UStaticMesh*> InStaticMeshes,
+	TArray<UObject*>& OutObjectsAdded
+	)
+{
+	for (UStaticMesh* StaticMesh: InStaticMeshes)
+	{
+		if (!ConvertSingleMeshToSkeletalMesh(InOptions, StaticMesh, OutObjectsAdded))
+		{
+			return false;
+		}
+	}
+	return true;
+}
+
+static void DoConversion(
+	UStaticMeshToSkeletalMeshConvertOptions* InOptions,
+	const TArray<UStaticMesh*>& InMeshesToConvert
+	)
+{
+	auto IsValidPathPart = [](const FString& InPathPart, FStringView InPathPartName, const FText& InMessage)-> bool
+	{
+		FText FailureReason;
+		const FText FailureContext = FText::FromStringView(InPathPartName);
+		if (!FName::IsValidXName(InPathPart, INVALID_OBJECTNAME_CHARACTERS INVALID_LONGPACKAGE_CHARACTERS, &FailureReason, &FailureContext))
+		{
+			ShowEditorMessage(ELogVerbosity::Error, InMessage, &FailureReason); 
+			return false;
+		}
+		return true;
+	};
+
+	FPackageName::EErrorCode PackagePathErrorCode;
+	if (!FPackageName::IsValidLongPackageName(InOptions->DestinationPath.Path, false, &PackagePathErrorCode))
+	{
+		const FText LogMessage = FPackageName::FormatErrorAsText(InOptions->DestinationPath.Path, PackagePathErrorCode);
+		ShowEditorMessage(ELogVerbosity::Error, LOCTEXT("InvalidDestinationPath", "Destination Path is Invalid"), &LogMessage); 
+		return;
+	}
+
+	if (!IsValidPathPart(InOptions->SkeletalMeshPrefixToAdd, TEXT("Skeletal Mesh Prefix"), LOCTEXT("InvalidSkeletalMeshPrefix", "Invalid Characters in Skeletal Mesh Prefix")) ||
+		!IsValidPathPart(InOptions->SkeletalMeshSuffixToAdd, TEXT("Skeletal Mesh Suffix"), LOCTEXT("InvalidSkeletalMeshSuffix", "Invalid Characters in Skeletal Mesh Suffix")))
+	{	
+		return;
+	}
+				
+	if (InOptions->SkeletonImportOption == EReferenceSkeletonImportOption::CreateNew)
+	{
+		if (!IsValidPathPart(InOptions->SkeletonPrefixToAdd, TEXT("Skeleton Prefix"), LOCTEXT("InvalidSkeletonPrefix", "Invalid Characters in Skeleton Prefix")) ||
+			!IsValidPathPart(InOptions->SkeletonSuffixToAdd, TEXT("Skeleton Suffix"), LOCTEXT("InvalidSkeletonSuffix", "Invalid Characters in Skeleton Suffix")))
+		{
+			return;
+		}
+	}
+
+	InOptions->SaveConfig();
+	
+	TArray<UObject*> ObjectsAdded;
+	if (ConvertMultipleMeshesToSkeletalMesh(InOptions, InMeshesToConvert, ObjectsAdded))
+	{
+		FAssetToolsModule::GetModule().Get().SyncBrowserToAssets(ObjectsAdded);
+	}
+	else
+	{
+		for (UObject* ObjectToDelete: ObjectsAdded)
+		{
+			FAssetRegistryModule::AssetDeleted(ObjectToDelete);
+			ObjectToDelete->Rename(nullptr, GetTransientPackage(), REN_DontCreateRedirectors | REN_NonTransactional);
+		}
+	}
+}
+
+void ConvertStaticMeshAssetsToSkeletalMeshesInteractive(
+	const TArray<FAssetData>& InStaticMeshAssets
+	)
+{
+	TArray<UStaticMesh*> MeshesToConvert;
+
+	for (const FAssetData& AssetData: InStaticMeshAssets)
+	{
+		UStaticMesh* StaticMesh = Cast<UStaticMesh>(AssetData.GetAsset());
+		if (!ensure(StaticMesh))
+		{
+			continue;
+		}
+
+		MeshesToConvert.Add(StaticMesh);
+	}
+	if (MeshesToConvert.IsEmpty())
+	{
+		return;
+	}
+
+	UStaticMeshToSkeletalMeshConvertOptions* Options = NewObject<UStaticMeshToSkeletalMeshConvertOptions>();
+	FPropertyEditorModule& PropertyEditorModule = FModuleManager::GetModuleChecked<FPropertyEditorModule>("PropertyEditor");
+	
+	// If the object paths don't resolve properly, the details view will still show something, which will be invalidated
+	// later.
+	if (Cast<USkeleton>(Options->Skeleton.ResolveObject()) == nullptr)
+	{
+		Options->Skeleton.Reset();
+	}
+	if (Cast<USkeletalMesh>(Options->SkeletalMesh.ResolveObject()) == nullptr)
+	{
+		Options->SkeletalMesh.Reset();
+	}
+
+	Options->DestinationPath.Path = FPackageName::GetLongPackagePath(MeshesToConvert[0]->GetPackage()->GetPathName());
+	
+	FDetailsViewArgs DetailsViewArgs;
+	DetailsViewArgs.bAllowSearch = false;
+	DetailsViewArgs.NameAreaSettings = FDetailsViewArgs::HideNameArea;
+	TSharedRef<IDetailsView> DetailsView = PropertyEditorModule.CreateDetailView(DetailsViewArgs);
+	DetailsView->SetObject(Options, /*bForceRefresh*/true);
+
+	TSharedRef<SCustomDialog> OptionsDialog = SNew(SCustomDialog)
+		.Title(LOCTEXT("OptionsDialogTitle", "Static Mesh Conversion Options"))
+		.Content()
+		[
+			SNew(SBox)
+			.MinDesiredWidth(450)
+			[
+				DetailsView
+			]
+		]
+	.Buttons({
+		SCustomDialog::FButton(LOCTEXT("DialogButtonConvert", "Convert"), FSimpleDelegate::CreateLambda([Options, MeshesToConvert]()
+		{
+			DoConversion(Options, MeshesToConvert);
+		})),
+		SCustomDialog::FButton(LOCTEXT("DialogButtonCancel", "Cancel"))
+	});
+
+	OptionsDialog->Show();
+}
+
+#undef LOCTEXT_NAMESPACE
diff --git a/Engine/Plugins/Animation/SkeletalMeshModelingTools/Source/SkeletalMeshModelingTools/Private/SkeletalMeshModelingToolsMeshConverter.h b/Engine/Plugins/Animation/SkeletalMeshModelingTools/Source/SkeletalMeshModelingTools/Private/SkeletalMeshModelingToolsMeshConverter.h
new file mode 100644
--- /dev/null
+++ b/Engine/Plugins/Animation/SkeletalMeshModelingTools/Source/SkeletalMeshModelingTools/Private/SkeletalMeshModelingToolsMeshConverter.h
@@ -0,0 +1,158 @@
+﻿// Copyright Epic Games, Inc. All Rights Reserved.
+
+#pragma once
+
+#include "BoneContainer.h"
+#include "Factories/Factory.h"
+#include "ReferenceSkeleton.h"
+#include "Animation/BoneSocketReference.h"
+#include "Interfaces/Interface_BoneReferenceSkeletonProvider.h"
+
+#include "SkeletalMeshModelingToolsMeshConverter.generated.h"
+
+
+class USkeleton;
+class UStaticMesh;
+struct FAssetData;
+
+UENUM()
+enum class ERootBonePositionReference
+{
+	Relative,
+	Absolute
+};
+
+
+UCLASS(Hidden, MinimalAPI)
+class USkeletonFromStaticMeshFactory :
+	public UFactory
+{
+	GENERATED_BODY()
+public:
+	USkeletonFromStaticMeshFactory(const FObjectInitializer& InObjectInitializer);
+
+	// UFactory overrides
+	virtual UObject* FactoryCreateNew(
+		UClass* InClass,
+		UObject* InParent,
+		FName InName,
+		EObjectFlags InFlags,
+		UObject* InContext,
+		FFeedbackContext* InWarn
+		) override;
+
+	virtual bool ShouldShowInNewMenu() const override { return false; }
+	
+	UPROPERTY()
+	TObjectPtr<UStaticMesh> StaticMesh;
+
+	UPROPERTY()
+	FVector RootPosition = {0.5, 0.5, 0.0};
+
+	UPROPERTY()
+	ERootBonePositionReference PositionReference = ERootBonePositionReference::Relative;
+};
+
+
+UCLASS(Hidden, MinimalAPI)
+class USkeletalMeshFromStaticMeshFactory :
+	public UFactory
+{
+	GENERATED_BODY()
+public:
+	USkeletalMeshFromStaticMeshFactory(const FObjectInitializer& InObjectInitializer);
+
+	// UFactory overrides
+	virtual UObject* FactoryCreateNew(
+		UClass* InClass,
+		UObject* InParent,
+		FName InName,
+		EObjectFlags InFlags,
+		UObject* InContext,
+		FFeedbackContext* InWarn
+		) override;
+	
+	virtual bool ShouldShowInNewMenu() const override { return false; }
+
+	FReferenceSkeleton ReferenceSkeleton;
+
+	UPROPERTY()
+	TObjectPtr<USkeleton> Skeleton;
+	
+	UPROPERTY()
+	TObjectPtr<UStaticMesh> StaticMesh;
+
+	UPROPERTY()
+	FName BindBoneName;
+};
+
+
+UENUM()
+enum class EReferenceSkeletonImportOption
+{
+	CreateNew UMETA(ToolTip="Create a new skeleton with a single root bone"),
+	UseExistingSkeleton UMETA(ToolTip="Use an already created skeleton asset"),
+	UseExistingSkeletalMesh UMETA(ToolTip="Use a skeleton use by another skeletal mesh asset")
+};
+
+UENUM()
+enum class ERootBonePlacementOptions
+{
+	BottomCenter UMETA(DisplayName="Bottom center", ToolTip="Bottom-center of the static mesh's bounding box"),
+	Center UMETA(DisplayName="Center", ToolTip="Center of the static mesh's bounding box"),
+	Origin UMETA(DisplayName="Origin", ToolTip="At the static mesh's origin (0,0,0)"),
+};
+
+
+
+UCLASS(config=EditorPerProjectUserSettings, HideCategories=Object, MinimalAPI)
+class UStaticMeshToSkeletalMeshConvertOptions :
+	public UObject,
+	public IBoneReferenceSkeletonProvider
+{
+	GENERATED_BODY()
+public:
+	UPROPERTY(EditAnywhere, Config, Category="Location", meta=(ContentDir))
+	FDirectoryPath DestinationPath;
+	
+
+	UPROPERTY(EditAnywhere, Config, Category="Skeleton", meta=(DisplayName="Creation Type"))
+	EReferenceSkeletonImportOption SkeletonImportOption = EReferenceSkeletonImportOption::CreateNew;
+
+	UPROPERTY(EditAnywhere, Config, Category="Skeleton", meta=(EditCondition="SkeletonImportOption==EReferenceSkeletonImportOption::CreateNew", EditConditionHides))
+	ERootBonePlacementOptions RootBonePlacement = ERootBonePlacementOptions::BottomCenter;
+	
+	UPROPERTY(EditAnywhere, Config, Category="Skeleton", meta=(AllowedClasses="/Script/Engine.Skeleton", EditCondition="SkeletonImportOption==EReferenceSkeletonImportOption::UseExistingSkeleton", EditConditionHides))
+	FSoftObjectPath Skeleton;
+
+	UPROPERTY(EditAnywhere, Config, Category="Skeleton", meta=(AllowedClasses="/Script/Engine.SkeletalMesh", EditCondition="SkeletonImportOption==EReferenceSkeletonImportOption::UseExistingSkeletalMesh", EditConditionHides))
+	FSoftObjectPath SkeletalMesh;
+
+	UPROPERTY(EditAnywhere, Config, Category="Skeleton", meta=(EditCondition="SkeletonImportOption!=EReferenceSkeletonImportOption::CreateNew", EditConditionHides))
+	FBoneReference BindingBoneName;
+
+	UPROPERTY(EditAnywhere, Config, Category="Naming")
+	FString PrefixToRemove = TEXT("SM_");
+	
+	UPROPERTY(EditAnywhere, Config, Category="Naming", meta=(DisplayName="Skeletal Mesh Prefix"))
+	FString SkeletalMeshPrefixToAdd = TEXT("SKM_");
+
+	UPROPERTY(EditAnywhere, Config, Category="Naming", meta=(DisplayName="Skeletal Mesh Suffix"))
+	FString SkeletalMeshSuffixToAdd;
+	
+	UPROPERTY(EditAnywhere, Config, Category="Naming", meta=(DisplayName="Skeleton Prefix", EditCondition="SkeletonImportOption==EReferenceSkeletonImportOption::CreateNew"))
+	FString SkeletonPrefixToAdd = TEXT("SK_");
+
+	UPROPERTY(EditAnywhere, Config, Category="Naming", meta=(DisplayName="Skeleton Suffix", EditCondition="SkeletonImportOption==EReferenceSkeletonImportOption::CreateNew"))
+	FString SkeletonSuffixToAdd;
+	
+	// IBoneReferenceSkeletonProvider implementation.
+	virtual USkeleton* GetSkeleton(bool& bInvalidSkeletonIsError, const IPropertyHandle* PropertyHandle) override;
+	
+	// UObject overrides.
+	virtual void PostEditChangeProperty(FPropertyChangedEvent& PropertyChangedEvent) override;
+};
+
+
+/** Converts a static mesh to a skeletal mesh, creating both a skeletal mesh and a skeleton assets that link up. */
+void ConvertStaticMeshAssetsToSkeletalMeshesInteractive(const TArray<FAssetData>& InStaticMeshAssets);
diff --git a/Engine/Plugins/Animation/SkeletalMeshModelingTools/Source/SkeletalMeshModelingTools/Private/SkeletalMeshModelingToolsModule.cpp b/Engine/Plugins/Animation/SkeletalMeshModelingTools/Source/SkeletalMeshModelingTools/Private/SkeletalMeshModelingToolsModule.cpp
new file mode 100644
--- /dev/null
+++ b/Engine/Plugins/Animation/SkeletalMeshModelingTools/Source/SkeletalMeshModelingTools/Private/SkeletalMeshModelingToolsModule.cpp
@@ -0,0 +1,261 @@
+// Copyright Epic Games, Inc. All Rights Reserved.
+
+#include "SkeletalMeshModelingToolsModule.h"
+
+#include "SkeletalMeshModelingToolsCommands.h"
+#include "SkeletalMeshModelingToolsEditorMode.h"
+#include "SkeletalMeshModelingToolsMeshConverter.h"
+#include "SkeletalMeshModelingToolsStyle.h"
+
+#include "ContentBrowserMenuContexts.h"
+#include "EditorModeManager.h"
+#include "Engine/SkeletalMesh.h"
+#include "ISkeletalMeshEditor.h"
+#include "ISkeletalMeshEditorModule.h"
+#include "ModelingToolsEditorModeStyle.h"
+#include "ModelingToolsManagerActions.h"
+#include "PersonaModule.h"
+#include "Modules/ModuleManager.h"
+#include "SkeletalMeshToolMenuContext.h"
+#include "ToolMenus.h"
+#include "DetailCustomization/SkeletonEditingToolPropertyCustomizations.h"
+#include "DetailCustomization/SkinWeightDetailCustomization.h"
+#include "Misc/ConfigCacheIni.h"
+#include "SkeletalMesh/SkeletonEditingTool.h"
+#include "SkeletalMesh/SkinWeightsPaintTool.h"
+#include "Styling/SlateIconFinder.h"
+#include "WorkflowOrientedApp/ApplicationMode.h"
+
+
+#define LOCTEXT_NAMESPACE "SkeletalMeshModelingToolsModule"
+
+DEFINE_LOG_CATEGORY(LogSkeletalMeshModelingTools);
+
+IMPLEMENT_MODULE(FSkeletalMeshModelingToolsModule, SkeletalMeshModelingTools);
+
+static const TCHAR* ConfigSection = TEXT("SkeletalMeshModelingTools");
+static const TCHAR* ConfigEditingModeActiveKey = TEXT("EditingModeActive");
+
+
+
+void FSkeletalMeshModelingToolsModule::StartupModule()
+{
+	FSkeletalMeshModelingToolsStyle::Register();
+	FSkeletalMeshModelingToolsCommands::Register();
+
+	UToolMenus::RegisterStartupCallback(
+		FSimpleMulticastDelegate::FDelegate::CreateRaw(this, &FSkeletalMeshModelingToolsModule::RegisterMenusAndToolbars));
+
+	ISkeletalMeshEditorModule& SkelMeshEditorModule = FModuleManager::Get().LoadModuleChecked<ISkeletalMeshEditorModule>("SkeletalMeshEditor");
+	// register toolbar extender with skeletal mesh editor
+	TArray<ISkeletalMeshEditorModule::FSkeletalMeshEditorToolbarExtender>& ToolbarExtenders = SkelMeshEditorModule.GetAllSkeletalMeshEditorToolbarExtenders();
+	ToolbarExtenders.Add(ISkeletalMeshEditorModule::FSkeletalMeshEditorToolbarExtender::CreateRaw(this, &FSkeletalMeshModelingToolsModule::ExtendSkelMeshEditorToolbar));
+	SkelMeshEditorExtenderHandle = ToolbarExtenders.Last().GetHandle();
+	// register post-init callback with skeletal mesh editor
+	TArray<ISkeletalMeshEditorModule::FOnSkeletalMeshEditorInitialized>& PostInitDelegates = SkelMeshEditorModule.GetPostEditorInitDelegates();
+	PostInitDelegates.Add(ISkeletalMeshEditorModule::FOnSkeletalMeshEditorInitialized::CreateRaw(this, &FSkeletalMeshModelingToolsModule::CheckEnableEditingToolModeOnOpen));
+	SkelMeshEditorPostInitHandle = PostInitDelegates.Last().GetHandle();
+
+	FCoreDelegates::OnPostEngineInit.AddRaw(this, &FSkeletalMeshModelingToolsModule::OnPostEngineInit);
+
+	RegisterPropertyCustomizations();
+}
+
+void FSkeletalMeshModelingToolsModule::ShutdownModule()
+{
+	UnregisterPropertyCustomizations();
+	
+	if (ISkeletalMeshEditorModule* SkelMeshEditorModule = FModuleManager::GetModulePtr<ISkeletalMeshEditorModule>("SkeletalMeshEditor"))
+	{
+		// un-register toolbar extender delegates
+		TArray<ISkeletalMeshEditorModule::FSkeletalMeshEditorToolbarExtender>& Extenders = SkelMeshEditorModule->GetAllSkeletalMeshEditorToolbarExtenders();
+		Extenders.RemoveAll([this](const ISkeletalMeshEditorModule::FSkeletalMeshEditorToolbarExtender& InDelegate) { return InDelegate.GetHandle() == SkelMeshEditorExtenderHandle; });
+
+		// un-register post-init delegates
+		TArray<ISkeletalMeshEditorModule::FOnSkeletalMeshEditorInitialized>& PostInitDelegates = SkelMeshEditorModule->GetPostEditorInitDelegates();
+		PostInitDelegates.RemoveAll([this](const ISkeletalMeshEditorModule::FOnSkeletalMeshEditorInitialized& InDelegate) { return InDelegate.GetHandle() == SkelMeshEditorPostInitHandle; });
+	}
+
+	UToolMenus::UnregisterOwner(this);
+
+	FSkeletalMeshModelingToolsCommands::Unregister();
+	FSkeletalMeshModelingToolsStyle::Unregister();
+	FModelingToolsManagerCommands::Unregister();
+	
+	FCoreDelegates::OnPostEngineInit.RemoveAll(this);
+	FSkeletalMeshModelingToolsActionCommands::UnregisterAllToolActions();
+}
+
+
+void FSkeletalMeshModelingToolsModule::RegisterMenusAndToolbars()
+{
+	FToolMenuOwnerScoped OwnerScoped(this);
+
+	if (UToolMenu* ToolMenu = UToolMenus::Get()->ExtendMenu("AssetEditor.SkeletalMeshEditor.ToolBar"))
+	{
+		FToolMenuSection& Section = ToolMenu->FindOrAddSection("SkeletalMesh");
+		Section.AddDynamicEntry("ToggleModelingToolsMode", FNewToolMenuSectionDelegate::CreateLambda([this](FToolMenuSection& InSection) {
+			const USkeletalMeshToolMenuContext* Context = InSection.FindContext<USkeletalMeshToolMenuContext>();
+			if (Context && Context->SkeletalMeshEditor.IsValid())
+			{
+				InSection.AddEntry(FToolMenuEntry::InitToolBarButton(
+				    FSkeletalMeshModelingToolsCommands::Get().ToggleEditingToolsMode,
+					LOCTEXT("SkeletalMeshEditorModelingMode", "Editing Tools"),
+				    LOCTEXT("SkeletalMeshEditorModelingModeTooltip", "Opens the Editing Tools palette that provides selected skeletal mesh modification tools, including skeleton, skin weights and attributes editing."),
+				    FSlateIcon("ModelingToolsStyle", "LevelEditor.ModelingToolsMode")));
+			}
+		}));
+	}
+
+	// Extend the asset browser for static meshes to include the conversion to skelmesh.
+	if (UToolMenu* ToolMenu = UToolMenus::Get()->ExtendMenu("ContentBrowser.AssetContextMenu.StaticMesh"))
+	{
+		FToolMenuSection& Section = ToolMenu->FindOrAddSection("GetAssetActions");
+		Section.AddDynamicEntry("ConvertToSkeletalMesh", FNewToolMenuSectionDelegate::CreateLambda([](FToolMenuSection& InSection)
+		{
+			const UContentBrowserAssetContextMenuContext* Context = InSection.FindContext<UContentBrowserAssetContextMenuContext>();
+			if (Context && Context->bCanBeModified)
+			{
+				const TAttribute<FText> Label = LOCTEXT("ConvertToSkeletalMesh", "Convert to Skeletal Mesh");
+				const TAttribute<FText> ToolTip = LOCTEXT("ConvertToSkeletalMeshToolTip", "Converts a non-Nanite static mesh to a skeletal mesh with an associated skeleton.");
+				const FSlateIcon Icon = FSlateIconFinder::FindIconForClass(USkeletalMesh::StaticClass());
+				const FToolMenuExecuteAction UIAction = FToolMenuExecuteAction::CreateLambda([](const FToolMenuContext& InMenuContext)
+				{
+					if (const UContentBrowserAssetContextMenuContext* Context = InMenuContext.FindContext<UContentBrowserAssetContextMenuContext>())
+					{
+						ConvertStaticMeshAssetsToSkeletalMeshesInteractive(
+							Context->GetSelectedAssetsOfType(UStaticMesh::StaticClass()));
+					}
+				});
+
+				InSection.AddMenuEntry("ConvertToSkeletalMesh", Label, ToolTip, Icon, UIAction);
+			}
+		}));
+	}
+}
+
+
+TSharedRef<FExtender> FSkeletalMeshModelingToolsModule::ExtendSkelMeshEditorToolbar(const TSharedRef<FUICommandList> InCommandList, TSharedRef<ISkeletalMeshEditor> InSkeletalMeshEditor)
+{
+	// Add toolbar extender
+	TSharedPtr<FExtender> ToolbarExtender = MakeShareable(new FExtender);
+
+	TWeakPtr<ISkeletalMeshEditor> EditorPtr(InSkeletalMeshEditor);
+
+	InCommandList->MapAction(FSkeletalMeshModelingToolsCommands::Get().ToggleEditingToolsMode,
+	    FExecuteAction::CreateRaw(this, &FSkeletalMeshModelingToolsModule::OnToggleEditingToolsMode, EditorPtr),
+	    FCanExecuteAction(),
+	    FIsActionChecked::CreateRaw(this, &FSkeletalMeshModelingToolsModule::IsEditingToolModeActive, EditorPtr));
+
+	return ToolbarExtender.ToSharedRef();
+}
+
+
+bool FSkeletalMeshModelingToolsModule::IsEditingToolModeActive(TWeakPtr<ISkeletalMeshEditor> InSkeletalMeshEditor) const
+{
+	TSharedPtr<ISkeletalMeshEditor> SkeletalMeshEditor = InSkeletalMeshEditor.Pin();
+	return SkeletalMeshEditor.IsValid() && SkeletalMeshEditor->GetEditorModeManager().IsModeActive(USkeletalMeshModelingToolsEditorMode::Id);
+}
+
+
+void FSkeletalMeshModelingToolsModule::OnToggleEditingToolsMode(TWeakPtr<ISkeletalMeshEditor> InSkeletalMeshEditor)
+{
+	TSharedPtr<ISkeletalMeshEditor> SkeletalMeshEditor = InSkeletalMeshEditor.Pin();
+	if (SkeletalMeshEditor.IsValid())
+	{
+		FEditorModeTools& EditorModeManager = SkeletalMeshEditor->GetEditorModeManager();
+		if (!IsEditingToolModeActive(InSkeletalMeshEditor))
+		{
+			EditorModeManager.ActivateMode(USkeletalMeshModelingToolsEditorMode::Id, true);
+
+			// bind notifications to skeletal mesh modeling editor mode
+			UEdMode* ActiveMode = EditorModeManager.GetActiveScriptableMode(USkeletalMeshModelingToolsEditorMode::Id);
+			if (USkeletalMeshModelingToolsEditorMode* SkeletalMeshMode = CastChecked<USkeletalMeshModelingToolsEditorMode>(ActiveMode))
+			{
+				SkeletalMeshMode->SetEditorBinding(SkeletalMeshEditor);
+			}
+		}
+		else
+		{
+			EditorModeManager.DeactivateMode(USkeletalMeshModelingToolsEditorMode::Id);
+		}
+
+		// Update the stored state of the editing tools active state.
+		GConfig->SetBool(ConfigSection, ConfigEditingModeActiveKey, IsEditingToolModeActive(InSkeletalMeshEditor), GEditorPerProjectIni);
+
+		// make sure SkeletonSelection is active when toggling the mode, as they are compatible.
+		// it will be deactivated later when entering a tool 
+		if (!EditorModeManager.IsModeActive(FPersonaEditModes::SkeletonSelection))
+		{
+			EditorModeManager.ActivateMode(FPersonaEditModes::SkeletonSelection);
+		}
+	}
+}
+
+void FSkeletalMeshModelingToolsModule::CheckEnableEditingToolModeOnOpen(TWeakPtr<ISkeletalMeshEditor> InSkeletalMeshEditor)
+{
+	bool bEditingModeActive = true;
+	GConfig->GetBool(ConfigSection, ConfigEditingModeActiveKey, bEditingModeActive, GEditorPerProjectIni);
+	if (bEditingModeActive && !IsEditingToolModeActive(InSkeletalMeshEditor))
+	{
+		OnToggleEditingToolsMode(InSkeletalMeshEditor);
+	}
+}
+
+void FSkeletalMeshModelingToolsModule::OnPostEngineInit()
+{
+	FSkeletalMeshModelingToolsActionCommands::RegisterAllToolActions();
+	FModelingToolsEditorModeStyle::Initialize();
+	FModelingToolsManagerCommands::Register();
+}
+
+void FSkeletalMeshModelingToolsModule::RegisterPropertyCustomizations()
+{
+	FPropertyEditorModule& PropertyModule = FModuleManager::LoadModuleChecked<FPropertyEditorModule>("PropertyEditor");
+
+	auto RegisterPropertyCustomization = [&](FName InStructName, auto InCustomizationFactory)
+	{
+		PropertyModule.RegisterCustomPropertyTypeLayout(
+			InStructName, 
+			FOnGetPropertyTypeCustomizationInstance::CreateStatic(InCustomizationFactory)
+			);
+		CustomizedProperties.Add(InStructName);
+	};
+
+	auto RegisterDetailCustomization = [&](FName InStructName, auto InCustomizationFactory)
+	{
+		PropertyModule.RegisterCustomClassLayout(
+			InStructName,
+			FOnGetDetailCustomizationInstance::CreateStatic(InCustomizationFactory)
+		);
+		CustomizedClasses.Add(InStructName);
+	};
+
+	
+	RegisterDetailCustomization(USkinWeightsPaintToolProperties::StaticClass()->GetFName(), &FSkinWeightDetailCustomization::MakeInstance);
+	RegisterDetailCustomization(USkeletonEditingTool::StaticClass()->GetFName(), &FSkeletonEditingToolDetailCustomization::MakeInstance);
+	RegisterDetailCustomization(USkeletonEditingProperties::StaticClass()->GetFName(), &FSkeletonEditingPropertiesDetailCustomization::MakeInstance);
+	RegisterDetailCustomization(UMirroringProperties::StaticClass()->GetFName(), &FMirroringPropertiesDetailCustomization::MakeInstance);
+	RegisterDetailCustomization(UOrientingProperties::StaticClass()->GetFName(), &FOrientingPropertiesDetailCustomization::MakeInstance);
+	RegisterDetailCustomization(UProjectionProperties::StaticClass()->GetFName(), &FProjectionPropertiesDetailCustomization::MakeInstance);
+
+	PropertyModule.NotifyCustomizationModuleChanged();
+}
+
+void FSkeletalMeshModelingToolsModule::UnregisterPropertyCustomizations()
+{
+	if (FPropertyEditorModule* PropertyModule = FModuleManager::GetModulePtr<FPropertyEditorModule>("PropertyEditor"))
+	{
+		for (const FName& PropertyName: CustomizedProperties)
+		{
+			PropertyModule->UnregisterCustomPropertyTypeLayout(PropertyName);
+		}
+		for (const FName& ClassName : CustomizedClasses)
+		{
+			PropertyModule->UnregisterCustomClassLayout(ClassName);
+		}
+		PropertyModule->NotifyCustomizationModuleChanged();
+	}
+}
+
+#undef LOCTEXT_NAMESPACE
diff --git a/Engine/Plugins/Animation/SkeletalMeshModelingTools/Source/SkeletalMeshModelingTools/Private/SkeletalMeshModelingToolsModule.h b/Engine/Plugins/Animation/SkeletalMeshModelingTools/Source/SkeletalMeshModelingTools/Private/SkeletalMeshModelingToolsModule.h
new file mode 100644
--- /dev/null
+++ b/Engine/Plugins/Animation/SkeletalMeshModelingTools/Source/SkeletalMeshModelingTools/Private/SkeletalMeshModelingToolsModule.h
@@ -0,0 +1,49 @@
+// Copyright Epic Games, Inc. All Rights Reserved.
+
+#pragma once
+
+#include "CoreMinimal.h"
+#include "Modules/ModuleInterface.h"
+
+
+DECLARE_LOG_CATEGORY_EXTERN(LogSkeletalMeshModelingTools, Log, All);
+
+
+class FApplicationMode;
+class ISkeletalMeshEditor;
+class FUICommandList;
+class FExtender;
+
+class FSkeletalMeshModelingToolsModule : 
+	public IModuleInterface
+{
+public:
+	// IModuleInterface implementations
+	void StartupModule() override;
+	void ShutdownModule() override;
+
+private:
+	// Extend toolbars and menus
+	void RegisterMenusAndToolbars();
+
+	TSharedRef<FExtender> ExtendSkelMeshEditorToolbar(const TSharedRef<FUICommandList> InCommandList, TSharedRef<ISkeletalMeshEditor> InSkeletalMeshEditor);
+
+
+	bool IsEditingToolModeActive(TWeakPtr<ISkeletalMeshEditor> InSkeletalMeshEditor) const;
+	void OnToggleEditingToolsMode(TWeakPtr<ISkeletalMeshEditor> InSkeletalMeshEditor);
+	void CheckEnableEditingToolModeOnOpen(TWeakPtr<ISkeletalMeshEditor> InSkeletalMeshEditor);
+
+	void OnPostEngineInit();
+	
+	void RegisterPropertyCustomizations();
+	void UnregisterPropertyCustomizations();
+
+	TArray<FName> CustomizedProperties;
+	TArray<FName> CustomizedClasses;
+	
+	// The handle for the extender delegate we added. Needed for clean module shutdown.
+	FDelegateHandle SkelMeshEditorExtenderHandle;
+	FDelegateHandle SkelMeshEditorPostInitHandle;
+
+	TArray<TWeakPtr<FApplicationMode>> RegisteredApplicationModes;
+};
diff --git a/Engine/Plugins/Animation/SkeletalMeshModelingTools/Source/SkeletalMeshModelingTools/Private/SkeletalMeshModelingToolsStyle.cpp b/Engine/Plugins/Animation/SkeletalMeshModelingTools/Source/SkeletalMeshModelingTools/Private/SkeletalMeshModelingToolsStyle.cpp
new file mode 100644
--- /dev/null
+++ b/Engine/Plugins/Animation/SkeletalMeshModelingTools/Source/SkeletalMeshModelingTools/Private/SkeletalMeshModelingToolsStyle.cpp
@@ -0,0 +1,32 @@
+// Copyright Epic Games, Inc. All Rights Reserved.
+
+#include "SkeletalMeshModelingToolsStyle.h"
+
+#include "Styling/SlateStyleMacros.h"
+#include "Styling/SlateStyleRegistry.h"
+
+
+FSkeletalMeshModelingToolsStyle::FSkeletalMeshModelingToolsStyle() : 
+	FSlateStyleSet("SkeletalMeshModelingToolsStyle")
+{
+	
+}
+
+
+FSkeletalMeshModelingToolsStyle& FSkeletalMeshModelingToolsStyle::Get()
+{
+	static FSkeletalMeshModelingToolsStyle Instance;
+	return Instance;
+}
+
+
+void FSkeletalMeshModelingToolsStyle::Register()
+{
+	FSlateStyleRegistry::RegisterSlateStyle(Get());
+}
+
+
+void FSkeletalMeshModelingToolsStyle::Unregister()
+{
+	FSlateStyleRegistry::UnRegisterSlateStyle(Get());
+}
diff --git a/Engine/Plugins/Animation/SkeletalMeshModelingTools/Source/SkeletalMeshModelingTools/Private/SkeletalMeshModelingToolsStyle.h b/Engine/Plugins/Animation/SkeletalMeshModelingTools/Source/SkeletalMeshModelingTools/Private/SkeletalMeshModelingToolsStyle.h
new file mode 100644
--- /dev/null
+++ b/Engine/Plugins/Animation/SkeletalMeshModelingTools/Source/SkeletalMeshModelingTools/Private/SkeletalMeshModelingToolsStyle.h
@@ -0,0 +1,22 @@
+// Copyright Epic Games, Inc. All Rights Reserved.
+
+#pragma once
+
+#include "Styling/SlateStyle.h"
+
+
+class FSkeletalMeshModelingToolsStyle
+    : public FSlateStyleSet
+{
+public:
+	static FSkeletalMeshModelingToolsStyle& Get();
+
+protected:
+	friend class FSkeletalMeshModelingToolsModule;
+
+	static void Register();
+	static void Unregister();
+
+private:
+	FSkeletalMeshModelingToolsStyle();
+};
diff --git a/Engine/Plugins/Animation/SkeletalMeshModelingTools/Source/SkeletalMeshModelingTools/Public/SkeletalMeshModelingModeToolExtensions.h b/Engine/Plugins/Animation/SkeletalMeshModelingTools/Source/SkeletalMeshModelingTools/Public/SkeletalMeshModelingModeToolExtensions.h
new file mode 100644
--- /dev/null
+++ b/Engine/Plugins/Animation/SkeletalMeshModelingTools/Source/SkeletalMeshModelingTools/Public/SkeletalMeshModelingModeToolExtensions.h
@@ -0,0 +1,25 @@
+// Copyright Epic Games, Inc. All Rights Reserved.
+
+#pragma once
+
+#include "Features/IModularFeature.h"
+#include "ModelingModeToolExtensions.h"
+#include "Internationalization/Text.h"
+
+/**
+ * ISkeletalMeshModelingModeToolExtension uses the IModularFeature API to allow a Plugin to provide
+ * a set of InteractiveTool's to be exposed in SkeletalMesh Modeling Mode. The Tools will be 
+ * included in a section of the SkeletalMesh Modeling Mode tool list, based on GetToolSectionName().
+ * 
+ */
+class ISkeletalMeshModelingModeToolExtension : public IModelingModeToolExtension
+{
+public:
+	virtual ~ISkeletalMeshModelingModeToolExtension() override {}
+
+	static FName GetModularFeatureName()
+	{
+		static FName FeatureName( TEXT("SkeletalMeshModelingModeToolExtension") );
+		return FeatureName;
+	}
+};
diff --git a/Engine/Plugins/Animation/SkeletalMeshModelingTools/Source/SkeletalMeshModelingTools/SkeletalMeshModelingTools.Build.cs b/Engine/Plugins/Animation/SkeletalMeshModelingTools/Source/SkeletalMeshModelingTools/SkeletalMeshModelingTools.Build.cs
new file mode 100644
--- /dev/null
+++ b/Engine/Plugins/Animation/SkeletalMeshModelingTools/Source/SkeletalMeshModelingTools/SkeletalMeshModelingTools.Build.cs
@@ -0,0 +1,44 @@
+﻿// Copyright Epic Games, Inc. All Rights Reserved.
+
+using UnrealBuildTool;
+
+public class SkeletalMeshModelingTools : ModuleRules
+{
+	public SkeletalMeshModelingTools(ReadOnlyTargetRules Target) : base(Target)
+	{
+		PrivateDependencyModuleNames.AddRange(
+			new string[]
+			{
+				"ContentBrowser",
+				"Core",
+				"CoreUObject",
+				"EditorFramework",
+				"EditorInteractiveToolsFramework",
+				"Engine",
+				"InputCore",
+				"InteractiveToolsFramework",
+				"MeshModelingTools",
+				"MeshModelingToolsEditorOnly",
+				"ModelingComponentsEditorOnly",
+				"ModelingToolsEditorMode",
+				"SkeletalMeshEditor",
+				"SkeletalMeshUtilitiesCommon",
+				"Slate",
+				"SlateCore",
+				"ToolMenus",
+				"ToolWidgets",
+				"UnrealEd",
+				"StatusBar",
+				"PropertyEditor",
+				"AnimationCore",
+				"AnimationWidgets",
+				"ApplicationCore",
+				"WidgetRegistration",
+				"SkeletalMeshModifiers",
+				"Persona"
+			}
+		);
+
+		ShortName = "SkelMeshModTools";
+	}
+}
diff --git a/Engine/Plugins/Experimental/Animation/SkeletalMeshModelingTools/SkeletalMeshModelingTools.uplugin b/Engine/Plugins/Experimental/Animation/SkeletalMeshModelingTools/SkeletalMeshModelingTools.uplugin
deleted file mode 100644
diff --git a/Engine/Plugins/Experimental/Animation/SkeletalMeshModelingTools/Source/SkeletalMeshModelingTools/Private/DetailCustomization/SkeletonEditingToolPropertyCustomizations.cpp b/Engine/Plugins/Experimental/Animation/SkeletalMeshModelingTools/Source/SkeletalMeshModelingTools/Private/DetailCustomization/SkeletonEditingToolPropertyCustomizations.cpp
deleted file mode 100644
diff --git a/Engine/Plugins/Experimental/Animation/SkeletalMeshModelingTools/Source/SkeletalMeshModelingTools/Private/DetailCustomization/SkeletonEditingToolPropertyCustomizations.h b/Engine/Plugins/Experimental/Animation/SkeletalMeshModelingTools/Source/SkeletalMeshModelingTools/Private/DetailCustomization/SkeletonEditingToolPropertyCustomizations.h
deleted file mode 100644
diff --git a/Engine/Plugins/Experimental/Animation/SkeletalMeshModelingTools/Source/SkeletalMeshModelingTools/Private/DetailCustomization/SkinWeightDetailCustomization.cpp b/Engine/Plugins/Experimental/Animation/SkeletalMeshModelingTools/Source/SkeletalMeshModelingTools/Private/DetailCustomization/SkinWeightDetailCustomization.cpp
deleted file mode 100644
diff --git a/Engine/Plugins/Experimental/Animation/SkeletalMeshModelingTools/Source/SkeletalMeshModelingTools/Private/DetailCustomization/SkinWeightDetailCustomization.h b/Engine/Plugins/Experimental/Animation/SkeletalMeshModelingTools/Source/SkeletalMeshModelingTools/Private/DetailCustomization/SkinWeightDetailCustomization.h
deleted file mode 100644
diff --git a/Engine/Plugins/Experimental/Animation/SkeletalMeshModelingTools/Source/SkeletalMeshModelingTools/Private/SReferenceSkeletonTree.cpp b/Engine/Plugins/Experimental/Animation/SkeletalMeshModelingTools/Source/SkeletalMeshModelingTools/Private/SReferenceSkeletonTree.cpp
deleted file mode 100644
diff --git a/Engine/Plugins/Experimental/Animation/SkeletalMeshModelingTools/Source/SkeletalMeshModelingTools/Private/SReferenceSkeletonTree.h b/Engine/Plugins/Experimental/Animation/SkeletalMeshModelingTools/Source/SkeletalMeshModelingTools/Private/SReferenceSkeletonTree.h
deleted file mode 100644
diff --git a/Engine/Plugins/Experimental/Animation/SkeletalMeshModelingTools/Source/SkeletalMeshModelingTools/Private/SkeletalMeshEditorUtils.cpp b/Engine/Plugins/Experimental/Animation/SkeletalMeshModelingTools/Source/SkeletalMeshModelingTools/Private/SkeletalMeshEditorUtils.cpp
deleted file mode 100644
diff --git a/Engine/Plugins/Experimental/Animation/SkeletalMeshModelingTools/Source/SkeletalMeshModelingTools/Private/SkeletalMeshEditorUtils.h b/Engine/Plugins/Experimental/Animation/SkeletalMeshModelingTools/Source/SkeletalMeshModelingTools/Private/SkeletalMeshEditorUtils.h
deleted file mode 100644
diff --git a/Engine/Plugins/Experimental/Animation/SkeletalMeshModelingTools/Source/SkeletalMeshModelingTools/Private/SkeletalMeshGizmoUtils.cpp b/Engine/Plugins/Experimental/Animation/SkeletalMeshModelingTools/Source/SkeletalMeshModelingTools/Private/SkeletalMeshGizmoUtils.cpp
deleted file mode 100644
diff --git a/Engine/Plugins/Experimental/Animation/SkeletalMeshModelingTools/Source/SkeletalMeshModelingTools/Private/SkeletalMeshGizmoUtils.h b/Engine/Plugins/Experimental/Animation/SkeletalMeshModelingTools/Source/SkeletalMeshModelingTools/Private/SkeletalMeshGizmoUtils.h
deleted file mode 100644
diff --git a/Engine/Plugins/Experimental/Animation/SkeletalMeshModelingTools/Source/SkeletalMeshModelingTools/Private/SkeletalMeshModelingToolsCommands.cpp b/Engine/Plugins/Experimental/Animation/SkeletalMeshModelingTools/Source/SkeletalMeshModelingTools/Private/SkeletalMeshModelingToolsCommands.cpp
deleted file mode 100644
diff --git a/Engine/Plugins/Experimental/Animation/SkeletalMeshModelingTools/Source/SkeletalMeshModelingTools/Private/SkeletalMeshModelingToolsCommands.h b/Engine/Plugins/Experimental/Animation/SkeletalMeshModelingTools/Source/SkeletalMeshModelingTools/Private/SkeletalMeshModelingToolsCommands.h
deleted file mode 100644
diff --git a/Engine/Plugins/Experimental/Animation/SkeletalMeshModelingTools/Source/SkeletalMeshModelingTools/Private/SkeletalMeshModelingToolsEditorMode.cpp b/Engine/Plugins/Experimental/Animation/SkeletalMeshModelingTools/Source/SkeletalMeshModelingTools/Private/SkeletalMeshModelingToolsEditorMode.cpp
deleted file mode 100644
diff --git a/Engine/Plugins/Experimental/Animation/SkeletalMeshModelingTools/Source/SkeletalMeshModelingTools/Private/SkeletalMeshModelingToolsEditorMode.h b/Engine/Plugins/Experimental/Animation/SkeletalMeshModelingTools/Source/SkeletalMeshModelingTools/Private/SkeletalMeshModelingToolsEditorMode.h
deleted file mode 100644
diff --git a/Engine/Plugins/Experimental/Animation/SkeletalMeshModelingTools/Source/SkeletalMeshModelingTools/Private/SkeletalMeshModelingToolsEditorModeToolkit.cpp b/Engine/Plugins/Experimental/Animation/SkeletalMeshModelingTools/Source/SkeletalMeshModelingTools/Private/SkeletalMeshModelingToolsEditorModeToolkit.cpp
deleted file mode 100644
diff --git a/Engine/Plugins/Experimental/Animation/SkeletalMeshModelingTools/Source/SkeletalMeshModelingTools/Private/SkeletalMeshModelingToolsEditorModeToolkit.h b/Engine/Plugins/Experimental/Animation/SkeletalMeshModelingTools/Source/SkeletalMeshModelingTools/Private/SkeletalMeshModelingToolsEditorModeToolkit.h
deleted file mode 100644
diff --git a/Engine/Plugins/Experimental/Animation/SkeletalMeshModelingTools/Source/SkeletalMeshModelingTools/Private/SkeletalMeshModelingToolsMeshConverter.cpp b/Engine/Plugins/Experimental/Animation/SkeletalMeshModelingTools/Source/SkeletalMeshModelingTools/Private/SkeletalMeshModelingToolsMeshConverter.cpp
deleted file mode 100644
diff --git a/Engine/Plugins/Experimental/Animation/SkeletalMeshModelingTools/Source/SkeletalMeshModelingTools/Private/SkeletalMeshModelingToolsMeshConverter.h b/Engine/Plugins/Experimental/Animation/SkeletalMeshModelingTools/Source/SkeletalMeshModelingTools/Private/SkeletalMeshModelingToolsMeshConverter.h
deleted file mode 100644
diff --git a/Engine/Plugins/Experimental/Animation/SkeletalMeshModelingTools/Source/SkeletalMeshModelingTools/Private/SkeletalMeshModelingToolsModule.cpp b/Engine/Plugins/Experimental/Animation/SkeletalMeshModelingTools/Source/SkeletalMeshModelingTools/Private/SkeletalMeshModelingToolsModule.cpp
deleted file mode 100644
diff --git a/Engine/Plugins/Experimental/Animation/SkeletalMeshModelingTools/Source/SkeletalMeshModelingTools/Private/SkeletalMeshModelingToolsModule.h b/Engine/Plugins/Experimental/Animation/SkeletalMeshModelingTools/Source/SkeletalMeshModelingTools/Private/SkeletalMeshModelingToolsModule.h
deleted file mode 100644
diff --git a/Engine/Plugins/Experimental/Animation/SkeletalMeshModelingTools/Source/SkeletalMeshModelingTools/Private/SkeletalMeshModelingToolsStyle.cpp b/Engine/Plugins/Experimental/Animation/SkeletalMeshModelingTools/Source/SkeletalMeshModelingTools/Private/SkeletalMeshModelingToolsStyle.cpp
deleted file mode 100644
diff --git a/Engine/Plugins/Experimental/Animation/SkeletalMeshModelingTools/Source/SkeletalMeshModelingTools/Private/SkeletalMeshModelingToolsStyle.h b/Engine/Plugins/Experimental/Animation/SkeletalMeshModelingTools/Source/SkeletalMeshModelingTools/Private/SkeletalMeshModelingToolsStyle.h
deleted file mode 100644
diff --git a/Engine/Plugins/Experimental/Animation/SkeletalMeshModelingTools/Source/SkeletalMeshModelingTools/Public/SkeletalMeshModelingModeToolExtensions.h b/Engine/Plugins/Experimental/Animation/SkeletalMeshModelingTools/Source/SkeletalMeshModelingTools/Public/SkeletalMeshModelingModeToolExtensions.h
deleted file mode 100644
diff --git a/Engine/Plugins/Experimental/Animation/SkeletalMeshModelingTools/Source/SkeletalMeshModelingTools/SkeletalMeshModelingTools.Build.cs b/Engine/Plugins/Experimental/Animation/SkeletalMeshModelingTools/Source/SkeletalMeshModelingTools/SkeletalMeshModelingTools.Build.cs
deleted file mode 100644
