Skeleton Editing Tool: add prefix/suffix + find & replace capabilities in the skeleton tree

diff --git a/Engine/Plugins/Animation/SkeletalMeshModelingTools/Source/SkeletalMeshModelingTools/Private/SReferenceSkeletonTree.cpp b/Engine/Plugins/Animation/SkeletalMeshModelingTools/Source/SkeletalMeshModelingTools/Private/SReferenceSkeletonTree.cpp
--- a/Engine/Plugins/Animation/SkeletalMeshModelingTools/Source/SkeletalMeshModelingTools/Private/SReferenceSkeletonTree.cpp
+++ b/Engine/Plugins/Animation/SkeletalMeshModelingTools/Source/SkeletalMeshModelingTools/Private/SReferenceSkeletonTree.cpp
@@ -55,9 +55,17 @@ FBoneElement::FBoneElement(const FName& InBoneName, TWeakObjectPtr<USkeletonModi
 	, WeakModifier(InModifier)
 {}
 
-void FBoneElement::RequestRename() const
+void FBoneElement::RequestRename(const ReferenceSkeletonTreeLocals::EBoneRenameType InRenameType) const
 {
-	OnRenameRequested.ExecuteIfBound();
+	OnRenameRequested.ExecuteIfBound(InRenameType);
+}
+
+SBoneItem::~SBoneItem()
+{
+	if (TSharedPtr<FBoneElement> BoneElement = WeakTreeElement.IsValid() ? WeakTreeElement.Pin() : nullptr)
+	{
+		BoneElement->OnRenameRequested.Unbind();
+	}
 }
 
 void SBoneItem::Construct(const FArguments& InArgs)
@@ -88,8 +96,6 @@ void SBoneItem::Construct(const FArguments& InArgs)
 		]
 	];
 
-	TSharedPtr< SInlineEditableTextBlock > InlineWidget;
-	
 	RowBox->AddSlot()
     .AutoWidth()
 	.Padding(4, 0, 0, 0)
@@ -103,9 +109,26 @@ void SBoneItem::Construct(const FArguments& InArgs)
 		.OnVerifyTextChanged(this, &SBoneItem::OnVerifyNameChanged)
 		.OnTextCommitted(InArgs._Delegates.OnBoneNameCommitted)
 		.MultiLine(false)
+    	.OnEnterEditingMode(this, &SBoneItem::OnEnterEditingMode)
+		.OnExitEditingMode(this, &SBoneItem::OnExitEditingMode)
+    	.ToolTipText_Lambda([this]()
+		{
+    		if (bEditing)
+    		{
+    			return LOCTEXT("BoneItem_RenameTooltip",
+    				"use $name to add a prefix/suffix or /find/replace/ to replace a substring\n"
+    				" - \"spine_$name\" to add \"spine_\" as prefix.\n"
+    				" - \"$name_left\" to add \"_left\" as suffix.\n"
+    				" - \"/neck_/tail_/\" to replace \"neck_\" by \"tail_\".\n");	
+    		}
+    		return FText::GetEmpty();
+		})
     ];
 
-	TreeElement->OnRenameRequested.BindSP(InlineWidget.Get(), &SInlineEditableTextBlock::EnterEditingMode);
+	if (InlineWidget)
+	{
+		TreeElement->OnRenameRequested.BindRaw(this, &SBoneItem::HandleOnRenameRequested);	
+	}
 }
 
 FSlateColor SBoneItem::GetTextColor(FIsSelected InIsSelected)
@@ -119,6 +142,15 @@ FSlateColor SBoneItem::GetTextColor(FIsSelected InIsSelected)
 	return FSlateColor::UseForeground();
 }
 
+void SBoneItem::HandleOnRenameRequested(const ReferenceSkeletonTreeLocals::EBoneRenameType InRenameType)
+{
+	OptRenameType = InRenameType;
+	if (InlineWidget)
+	{
+		InlineWidget->EnterEditingMode();
+	}
+}
+
 bool SBoneItem::OnVerifyNameChanged(const FText& InText, FText& OutErrorMessage) const
 {
 	if (!WeakModifier.IsValid())
@@ -130,12 +162,52 @@ bool SBoneItem::OnVerifyNameChanged(const FText& InText, FText& OutErrorMessage)
 	const FReferenceSkeleton& ReferenceSkeleton = WeakModifier->GetReferenceSkeleton();
 
 	const FName NewName = FName(InText.ToString());
-	const int32 BoneIndex = ReferenceSkeleton.FindRawBoneIndex(NewName);
+	int32 BoneIndex = ReferenceSkeleton.FindRawBoneIndex(NewName);
+	if (BoneIndex != INDEX_NONE)
+	{
+		// changing case is allowed 
+		if (!ReferenceSkeleton.GetBoneName(BoneIndex).IsEqual(NewName, ENameCase::CaseSensitive))
+		{
+			BoneIndex = INDEX_NONE;
+		}
+	}
+	
 	return BoneIndex == INDEX_NONE;
 }
 
+void SBoneItem::OnEnterEditingMode()
+{
+	bEditing = true;
+}
+
+void SBoneItem::OnExitEditingMode()
+{
+	bEditing = false;
+	OptRenameType.Reset();
+}
+
 FText SBoneItem::GetName() const
 {
+	TSharedPtr<FBoneElement> BoneElement = WeakTreeElement.IsValid() ? WeakTreeElement.Pin() : nullptr;
+	ensure(BoneElement);
+	
+	if (bEditing && OptRenameType)
+	{
+		switch (*OptRenameType)
+		{
+		case ReferenceSkeletonTreeLocals::EBoneRenameType::Rename:
+			return FText::FromName(BoneElement->BoneName);
+		case ReferenceSkeletonTreeLocals::EBoneRenameType::SearchAndReplace:
+			return FText::FromString(TEXT("/to_search/to_replace/"));
+		case ReferenceSkeletonTreeLocals::EBoneRenameType::AddPrefix:
+			return FText::FromString(TEXT("prefix_$name"));
+		case ReferenceSkeletonTreeLocals::EBoneRenameType::AddSuffix:
+			return FText::FromString(TEXT("$name_suffix"));
+		default:
+			return FText::FromName(BoneElement->BoneName);
+		}
+	}
+	
 	return FText::FromName(WeakTreeElement.Pin()->BoneName);
 }
 
@@ -367,7 +439,19 @@ void SReferenceSkeletonTree::BindCommands()
         FCanExecuteAction::CreateSP(this, &SReferenceSkeletonTree::CanDeleteBone));
  
 	CommandList->MapAction(Commands.RenameBone,
-	FExecuteAction::CreateSP(this, &SReferenceSkeletonTree::HandleRenameBone),
+		FExecuteAction::CreateSP(this, &SReferenceSkeletonTree::HandleRenameBone, ReferenceSkeletonTreeLocals::EBoneRenameType::Rename),
+		FCanExecuteAction::CreateSP(this, &SReferenceSkeletonTree::CanRenameBone));
+
+	CommandList->MapAction(Commands.SearchAndReplace,
+		FExecuteAction::CreateSP(this, &SReferenceSkeletonTree::HandleRenameBone, ReferenceSkeletonTreeLocals::EBoneRenameType::SearchAndReplace),
+		FCanExecuteAction::CreateSP(this, &SReferenceSkeletonTree::CanRenameBone));
+
+	CommandList->MapAction(Commands.AddPrefix,
+		FExecuteAction::CreateSP(this, &SReferenceSkeletonTree::HandleRenameBone, ReferenceSkeletonTreeLocals::EBoneRenameType::AddPrefix),
+		FCanExecuteAction::CreateSP(this, &SReferenceSkeletonTree::CanRenameBone));
+
+	CommandList->MapAction(Commands.AddSuffix,
+		FExecuteAction::CreateSP(this, &SReferenceSkeletonTree::HandleRenameBone, ReferenceSkeletonTreeLocals::EBoneRenameType::AddSuffix),
 		FCanExecuteAction::CreateSP(this, &SReferenceSkeletonTree::CanRenameBone));
 	
 	CommandList->MapAction(Commands.UnParentBone,
@@ -623,14 +707,14 @@ ISkeletalMeshNotifier& SReferenceSkeletonTree::GetNotifier()
 	return *Notifier;
 }
 
-void SReferenceSkeletonTree::HandleRenameBone() const
+void SReferenceSkeletonTree::HandleRenameBone(const ReferenceSkeletonTreeLocals::EBoneRenameType InRenameType) const
 {
 	TArray<TSharedPtr<FBoneElement>> BoneElements = TreeView->GetSelectedItems();
 	if (BoneElements.IsEmpty())
 	{
 		return;
 	}
-	BoneElements[0]->RequestRename();
+	BoneElements[0]->RequestRename(InRenameType);
 }
 
 bool SReferenceSkeletonTree::CanRenameBone() const
@@ -653,21 +737,58 @@ void SReferenceSkeletonTree::OnBoneRenamed(const FName InOldName, const FName In
 
 void SReferenceSkeletonTree::OnNewBoneNameCommitted(const FText& InText, ETextCommit::Type InCommitType)
 {
+	if (InCommitType == ETextCommit::OnCleared)
+	{
+		return;
+	}
+	
 	TArray<FName> BoneNames;
 	GetSelectedBoneNames(BoneNames);
+
+	static const FString NamePattern(TEXT("$name"));
+	static constexpr TCHAR Slash(TEXT('/'));
 	
-	const FName NewName = FName(InText.ToString());
-	
-	if (BoneNames.IsEmpty() || NewName == NAME_None)
+	const FString NameStr = InText.ToString();
+	const FName NewName = FName(NameStr);
+
+	// wrong input
+	if (BoneNames.IsEmpty() || NewName == NAME_None || NameStr == NamePattern)
 	{
 		return;
 	}
 	
-	if (BoneNames.Num() == 1 && BoneNames[0] == NewName)
+	// same name
+	if (BoneNames[0].IsEqual(NewName, ENameCase::CaseSensitive))
 	{
 		return;
 	}
 
+	// look for prefix / suffix pattern 
+	const bool bAddPrefixSuffix = NameStr.Contains(NamePattern);
+	
+	// look for find & replace pattern 
+	int32 SlashIndex[3] = { INDEX_NONE, INDEX_NONE, INDEX_NONE };
+	for (int32 CharIndex = 0, NumSlash = 0; NumSlash < 3 && CharIndex < NameStr.Len(); ++CharIndex)
+	{
+		if (NameStr[CharIndex] == Slash)
+		{
+			SlashIndex[NumSlash++] = CharIndex;
+		}
+	}
+
+	const bool bFindReplace =
+		SlashIndex[0] == 0 && SlashIndex[2] == NameStr.Len()-1 && // starts and ends with /
+		SlashIndex[0] != SlashIndex[1] && SlashIndex[0] != SlashIndex[2]; // all three indices are different
+
+	if (bFindReplace)
+	{
+		if (bAddPrefixSuffix || (SlashIndex[1] - SlashIndex[0]) < 2)
+		{
+			// error
+			return;
+		}
+	}
+	
 	TArray<int32> BoneIndices;
 	const FReferenceSkeleton& ReferenceSkeleton = Modifier->GetReferenceSkeleton();
 	Algo::Transform(BoneNames, BoneIndices, [&](const FName& BoneName)
@@ -676,7 +797,36 @@ void SReferenceSkeletonTree::OnNewBoneNameCommitted(const FText& InText, ETextCo
 	});
 
 	TArray<FName> NewNames;
-	NewNames.Init(NewName, BoneNames.Num());
+
+	if (bFindReplace)
+	{
+		const FString Find = FString::ConstructFromPtrSize(&NameStr[SlashIndex[0]+1], SlashIndex[1]-SlashIndex[0]-1);
+		const FString Replace = FString::ConstructFromPtrSize(&NameStr[SlashIndex[1]+1], SlashIndex[2]-SlashIndex[1]-1);
+		
+		NewNames.Reserve(BoneNames.Num());
+		for (int32 Index = 0; Index < BoneNames.Num(); Index++)
+		{
+			FString NewBoneNameStr = BoneNames[Index].ToString().Replace(*Find, *Replace, ESearchCase::CaseSensitive);
+			NewNames.Add(*NewBoneNameStr);
+		} 
+	}
+	else
+	{
+		FString Prefix, Suffix;
+		if (bAddPrefixSuffix && NameStr.Split(NamePattern, &Prefix, &Suffix))
+		{
+			NewNames.Reserve(BoneNames.Num());
+			for (int32 Index = 0; Index < BoneNames.Num(); Index++)
+			{
+				FString NewBoneNameStr = FString::Printf( TEXT("%s%s%s"), *Prefix, *BoneNames[Index].ToString(), *Suffix);
+				NewNames.Add(*NewBoneNameStr);
+			}
+		}
+		else
+		{
+			NewNames.Init(NewName, BoneNames.Num());
+		}
+	}
 	
 	BeginChange();
 
@@ -919,20 +1069,28 @@ TSharedPtr<SWidget> SReferenceSkeletonTree::CreateContextMenu()
 	else
 	{
 		MenuBuilder.BeginSection("EditBones", LOCTEXT("EditBonesOperations", "Bones"));
-
-		MenuBuilder.AddMenuEntry(Commands.NewBone);
-		MenuBuilder.AddMenuEntry(Commands.RemoveBone);
-		MenuBuilder.AddMenuEntry(Commands.UnParentBone);
-		MenuBuilder.AddMenuEntry(Commands.RenameBone);
-
+		{
+			MenuBuilder.AddMenuEntry(Commands.NewBone);
+			MenuBuilder.AddMenuEntry(Commands.RemoveBone);
+			MenuBuilder.AddMenuEntry(Commands.UnParentBone);
+		}
+		MenuBuilder.EndSection();
+		
+		MenuBuilder.BeginSection("RenameBones", LOCTEXT("RenameBonesOperations", "Rename"));
+		{
+			MenuBuilder.AddMenuEntry(Commands.RenameBone);
+			MenuBuilder.AddMenuEntry(Commands.SearchAndReplace);
+			MenuBuilder.AddMenuEntry(Commands.AddPrefix);
+			MenuBuilder.AddMenuEntry(Commands.AddSuffix);
+		}
 		MenuBuilder.EndSection();
 
 		MenuBuilder.BeginSection("CopyPasteBones", LOCTEXT("CopyPasteBonesOperations", "Copy & Paste"));
-
-		MenuBuilder.AddMenuEntry(Commands.CopyBones);
-		MenuBuilder.AddMenuEntry(Commands.PasteBones);
-		MenuBuilder.AddMenuEntry(Commands.DuplicateBones);
-
+		{
+			MenuBuilder.AddMenuEntry(Commands.CopyBones);
+			MenuBuilder.AddMenuEntry(Commands.PasteBones);
+			MenuBuilder.AddMenuEntry(Commands.DuplicateBones);
+		}
 		MenuBuilder.EndSection();
 	}
 	
@@ -941,7 +1099,7 @@ TSharedPtr<SWidget> SReferenceSkeletonTree::CreateContextMenu()
 
 void SReferenceSkeletonTree::OnItemDoubleClicked(TSharedPtr<FBoneElement> InItem)
 {
-	InItem->RequestRename();
+	InItem->RequestRename(ReferenceSkeletonTreeLocals::EBoneRenameType::Rename);
 }
 
 void SReferenceSkeletonTree::OnSetExpansionRecursive(TSharedPtr<FBoneElement> InItem, bool bShouldBeExpanded)
diff --git a/Engine/Plugins/Animation/SkeletalMeshModelingTools/Source/SkeletalMeshModelingTools/Private/SReferenceSkeletonTree.h b/Engine/Plugins/Animation/SkeletalMeshModelingTools/Source/SkeletalMeshModelingTools/Private/SReferenceSkeletonTree.h
--- a/Engine/Plugins/Animation/SkeletalMeshModelingTools/Source/SkeletalMeshModelingTools/Private/SReferenceSkeletonTree.h
+++ b/Engine/Plugins/Animation/SkeletalMeshModelingTools/Source/SkeletalMeshModelingTools/Private/SReferenceSkeletonTree.h
@@ -15,6 +15,7 @@
 class FTextFilterExpressionEvaluator;
 class FUICommandList;
 class USkeletonModifier;
+class SInlineEditableTextBlock;
 
 namespace ReferenceSkeletonTreeLocals
 {
@@ -41,6 +42,14 @@ private:
 	TArray<int32> PostBoneTracker;
 };
 
+enum class EBoneRenameType : uint32
+{
+	Rename,
+	SearchAndReplace,
+	AddPrefix,
+	AddSuffix
+};
+	
 }
 
 class FBoneElement : public TSharedFromThis<FBoneElement>
@@ -56,9 +65,9 @@ public:
 	TSharedPtr<FBoneElement> Parent;
 	TArray<TSharedPtr<FBoneElement>> Children;
 
-	void RequestRename() const;
+	void RequestRename(const ReferenceSkeletonTreeLocals::EBoneRenameType InRenameType) const;
 	
-	DECLARE_DELEGATE(FOnRenameRequested);
+	DECLARE_DELEGATE_OneParam(FOnRenameRequested, const ReferenceSkeletonTreeLocals::EBoneRenameType);
 	FOnRenameRequested OnRenameRequested;
 
 	DECLARE_DELEGATE_TwoParams(FOnRenamed, const FName, const FName);
@@ -90,19 +99,30 @@ class SBoneItem : public SCompoundWidget
 		SLATE_EVENT( FIsSelected, IsSelected )
 	SLATE_END_ARGS()
 
+	virtual ~SBoneItem() override;
+	
 	// Slate construction function
 	void Construct(const FArguments& InArgs);
+
+	void HandleOnRenameRequested(const ReferenceSkeletonTreeLocals::EBoneRenameType InType);
 	
 private:
 
 	bool OnVerifyNameChanged(const FText& InText, FText& OutErrorMessage) const;
 
+	void OnEnterEditingMode();
+	void OnExitEditingMode();
+
 	static FSlateColor GetTextColor(FIsSelected InIsSelected);
 	
 	FText GetName() const;
 	
 	TWeakPtr<FBoneElement> WeakTreeElement;
 	TWeakObjectPtr<USkeletonModifier> WeakModifier;
+	TSharedPtr< SInlineEditableTextBlock > InlineWidget;
+	
+	bool bEditing = false;
+	TOptional<ReferenceSkeletonTreeLocals::EBoneRenameType> OptRenameType;
 };
 
 class FBoneItemDragDropOp : public FDecoratedDragDropOp
@@ -194,7 +214,7 @@ private:
 	// Create, Rename, Delete, Parent bones
 	void HandleNewBone();
 	bool CanAddNewBone() const;
-	void HandleRenameBone() const;
+	void HandleRenameBone(const ReferenceSkeletonTreeLocals::EBoneRenameType InRenameType) const;
 	bool CanRenameBone() const;
 	void HandleDeleteBone();
 	bool CanDeleteBone() const;
diff --git a/Engine/Plugins/Animation/SkeletalMeshModelingTools/Source/SkeletalMeshModelingTools/Private/SkeletalMeshModelingToolsCommands.cpp b/Engine/Plugins/Animation/SkeletalMeshModelingTools/Source/SkeletalMeshModelingTools/Private/SkeletalMeshModelingToolsCommands.cpp
--- a/Engine/Plugins/Animation/SkeletalMeshModelingTools/Source/SkeletalMeshModelingTools/Private/SkeletalMeshModelingToolsCommands.cpp
+++ b/Engine/Plugins/Animation/SkeletalMeshModelingTools/Source/SkeletalMeshModelingTools/Private/SkeletalMeshModelingToolsCommands.cpp
@@ -14,8 +14,12 @@ void FSkeletalMeshModelingToolsCommands::RegisterCommands()
 	UI_COMMAND(NewBone, "New Bone", "Add new Bone.", EUserInterfaceActionType::Button, FInputChord(EKeys::N, EModifierKey::Control));
 	UI_COMMAND(RemoveBone, "Remove Bone", "Remove selected bone(s).", EUserInterfaceActionType::Button, FInputChord(EKeys::Delete), FInputChord(EKeys::BackSpace));
 	UI_COMMAND(UnParentBone, "Unparent Bone", "Unparent selected bone(s).", EUserInterfaceActionType::Button, FInputChord(EKeys::P, EModifierKey::Shift));
-	UI_COMMAND(RenameBone, "Rename Bone", "Rename the selected bone.", EUserInterfaceActionType::Button, FInputChord(EKeys::F2));
 
+	UI_COMMAND(RenameBone, "Rename Bone", "Rename the selected bone.", EUserInterfaceActionType::Button, FInputChord(EKeys::F2));
+	UI_COMMAND(SearchAndReplace, "Search and Replace", "Search and Replace the selected bone(s) name(s).", EUserInterfaceActionType::Button, FInputChord());
+	UI_COMMAND(AddPrefix, "Add Prefix", "Add a prefix to the selected bone(s) name(s).", EUserInterfaceActionType::Button, FInputChord());
+	UI_COMMAND(AddSuffix, "Add Suffix", "Add a suffix to the selected bone(s) name(s).", EUserInterfaceActionType::Button, FInputChord());
+	
 	UI_COMMAND(CopyBones, "Copy Bone(s)", "Copy selected bone(s).", EUserInterfaceActionType::Button, FInputChord(EKeys::C, EModifierKey::Control));
 	UI_COMMAND(PasteBones, "Paste Bone(s)", "Paste selected bone(s).", EUserInterfaceActionType::Button, FInputChord(EKeys::V, EModifierKey::Control));
 	UI_COMMAND(DuplicateBones, "Duplicate Bone(s)", "Duplicate selected bone(s).", EUserInterfaceActionType::Button, FInputChord(EKeys::D, EModifierKey::Control));
diff --git a/Engine/Plugins/Animation/SkeletalMeshModelingTools/Source/SkeletalMeshModelingTools/Private/SkeletalMeshModelingToolsCommands.h b/Engine/Plugins/Animation/SkeletalMeshModelingTools/Source/SkeletalMeshModelingTools/Private/SkeletalMeshModelingToolsCommands.h
--- a/Engine/Plugins/Animation/SkeletalMeshModelingTools/Source/SkeletalMeshModelingTools/Private/SkeletalMeshModelingToolsCommands.h
+++ b/Engine/Plugins/Animation/SkeletalMeshModelingTools/Source/SkeletalMeshModelingTools/Private/SkeletalMeshModelingToolsCommands.h
@@ -31,6 +31,9 @@ public:
 	TSharedPtr<FUICommandInfo> RemoveBone;
 	TSharedPtr<FUICommandInfo> UnParentBone;
 	TSharedPtr<FUICommandInfo> RenameBone;
+	TSharedPtr<FUICommandInfo> SearchAndReplace;
+	TSharedPtr<FUICommandInfo> AddPrefix;
+	TSharedPtr<FUICommandInfo> AddSuffix;
 
 	// Copy & Paste commands
 	TSharedPtr<FUICommandInfo> CopyBones;
diff --git a/Engine/Plugins/Runtime/MeshModelingToolset/Source/SkeletalMeshModifiers/Private/SkeletonModifier.cpp b/Engine/Plugins/Runtime/MeshModelingToolset/Source/SkeletalMeshModifiers/Private/SkeletonModifier.cpp
--- a/Engine/Plugins/Runtime/MeshModelingToolset/Source/SkeletalMeshModifiers/Private/SkeletonModifier.cpp
+++ b/Engine/Plugins/Runtime/MeshModelingToolset/Source/SkeletalMeshModifiers/Private/SkeletonModifier.cpp
@@ -618,6 +618,12 @@ bool USkeletonModifier::CommitSkeletonToSkeletalMesh()
 		}
 
 		SkeletalMesh->SetSkeleton(NewSkeleton);
+		SkeletalMesh->MarkPackageDirty();
+
+		if (NewSkeleton->GetPreviewMesh() != SkeletalMesh)
+		{
+			NewSkeleton->SetPreviewMesh(SkeletalMesh.Get());
+		}
 	}
 	
 	const TArray<FMeshBoneInfo>& BoneInfos = ReferenceSkeleton->GetRawRefBoneInfo();
@@ -1495,13 +1501,16 @@ bool USkeletonModifier::RenameBones(const TArray<FName>& InOldBoneNames, const T
 	{
 		FReferenceSkeletonModifier Modifier(*ReferenceSkeleton, nullptr);
 
+		TArray<FName> NewBoneNames;
+		
 		const int32 NumBonesToRename = InOldBoneNames.Num();
 		for (int32 Index = 0; Index < NumBonesToRename; ++Index)
 		{
 			const FName& OldName = InOldBoneNames[Index];
-			const FName NewName = GetUniqueName(InNewBoneNames[Index]);
-			if (OldName != NAME_None && OldName != NewName)
+			const FName NewName = GetUniqueName(InNewBoneNames[Index], NewBoneNames);
+			if (OldName != NAME_None && !OldName.IsEqual(NewName, ENameCase::CaseSensitive))
 			{
+				NewBoneNames.Add(NewName);
 				Modifier.Rename(OldName, NewName);
 			}
 		}
@@ -1873,12 +1882,26 @@ FName USkeletonModifier::GetUniqueName(const FName InBoneName, const TArray<FNam
 	auto IsNameAvailable = [&](const FString& InNameStr)
 	{
 		const FName Name(*InNameStr);
+		
 		const int32 Index = ReferenceSkeleton->FindRawBoneIndex(Name);
-		if (Index != INDEX_NONE)
-		{
-			return false;
-		}
-
+        if (Index != INDEX_NONE)
+        {
+        	if (!ReferenceSkeleton->GetBoneName(Index).IsEqual(Name, ENameCase::CaseSensitive))
+        	{
+        		const bool bContainsByPredicate = InBoneNames.ContainsByPredicate([Name](const FName& BoneName)
+        		{
+        			return BoneName.IsEqual(Name, ENameCase::CaseSensitive);
+        		});
+        		
+        		if (!bContainsByPredicate)
+        		{
+        			return true;
+        		}
+        	}
+        	
+        	return false;
+        }
+		
 		if (InBoneNames.Contains(Name))
 		{
 			return false;
diff --git a/Engine/Source/Runtime/Engine/Private/Animation/ReferenceSkeleton.cpp b/Engine/Source/Runtime/Engine/Private/Animation/ReferenceSkeleton.cpp
--- a/Engine/Source/Runtime/Engine/Private/Animation/ReferenceSkeleton.cpp
+++ b/Engine/Source/Runtime/Engine/Private/Animation/ReferenceSkeleton.cpp
@@ -203,7 +203,8 @@ void FReferenceSkeleton::Remove(const FName InBoneName, const bool bRemoveChildr
 void FReferenceSkeleton::Rename(const FName InBoneName, const FName InNewName)
 {
 	const int32 RawBoneIndex = FindRawBoneIndex(InBoneName);
-	const int32 RawNewBoneIndex = FindRawBoneIndex(InNewName);
+	const bool bCaseSensitiveOnly = InBoneName == InNewName && !InBoneName.IsEqual(InNewName, ENameCase::CaseSensitive);
+	const int32 RawNewBoneIndex = bCaseSensitiveOnly ? INDEX_NONE : FindRawBoneIndex(InNewName);
 	if (RawBoneIndex == INDEX_NONE || RawNewBoneIndex != INDEX_NONE)
 	{
 		return;
