Skin Weight Paint Tool: ensure that tool & settings are safe to use
- as most of the UI uses lambdas, weak ptrs can be accessed while invalid

#rnx

diff --git a/Engine/Plugins/Animation/SkeletalMeshModelingTools/Source/SkeletalMeshModelingTools/Private/DetailCustomization/SkinWeightDetailCustomization.cpp b/Engine/Plugins/Animation/SkeletalMeshModelingTools/Source/SkeletalMeshModelingTools/Private/DetailCustomization/SkinWeightDetailCustomization.cpp
--- a/Engine/Plugins/Animation/SkeletalMeshModelingTools/Source/SkeletalMeshModelingTools/Private/DetailCustomization/SkinWeightDetailCustomization.cpp
+++ b/Engine/Plugins/Animation/SkeletalMeshModelingTools/Source/SkeletalMeshModelingTools/Private/DetailCustomization/SkinWeightDetailCustomization.cpp
@@ -66,15 +66,18 @@ void FSkinWeightDetailCustomization::CustomizeDetails(IDetailLayoutBuilder& Deta
 					"Bones: select and manipulate bones to preview deformations.\n"))
 			.Value_Lambda([this]()
 			{
-				return ToolSettings->EditingMode;
+				return ToolSettings.IsValid() ? ToolSettings->EditingMode : EWeightEditMode::Brush;
 			})
 			.OnValueChanged_Lambda([this](EWeightEditMode Mode)
 			{
-				ToolSettings->EditingMode = Mode;
-				ToolSettings->WeightTool->ToggleEditingMode();
-				if (CurrentDetailBuilder)
+				if (ToolSettings.IsValid())
 				{
-					CurrentDetailBuilder->ForceRefreshDetails();
+					ToolSettings->EditingMode = Mode;
+					ToolSettings->WeightTool->ToggleEditingMode();
+					if (CurrentDetailBuilder)
+					{
+						CurrentDetailBuilder->ForceRefreshDetails();
+					}
 				}
 			})
 			+SSegmentedControl<EWeightEditMode>::Slot(EWeightEditMode::Brush)
@@ -117,11 +120,14 @@ void FSkinWeightDetailCustomization::CustomizeDetails(IDetailLayoutBuilder& Deta
 					"Full Material: Displays normal mesh materials with textures.\n"))
 			.Value_Lambda([this]()
 			{
-				return ToolSettings->ColorMode;
+				return ToolSettings.IsValid() ? ToolSettings->ColorMode : EWeightColorMode::Greyscale;
 			})
 			.OnValueChanged_Lambda([this](EWeightColorMode Mode)
 			{
-				ToolSettings->SetColorMode(Mode);
+				if (ToolSettings.IsValid())
+				{
+					ToolSettings->SetColorMode(Mode);
+				}
 			})
 			+ SSegmentedControl<EWeightColorMode>::Slot(EWeightColorMode::Greyscale)
 			.Text(LOCTEXT("GreyscaleMode", "Greyscale"))
@@ -190,11 +196,14 @@ void FSkinWeightDetailCustomization::AddBrushUI(IDetailLayoutBuilder& DetailBuil
 				"Relax: applies the average of the connected (by edge) vertex weights to the new vertex weight, blended by the Strength."))
 			.Value_Lambda([this]()
 			{
-				return ToolSettings->BrushMode;
+				return ToolSettings.IsValid() ? ToolSettings->BrushMode : EWeightEditOperation::Add;
 			})
 			.OnValueChanged_Lambda([this](EWeightEditOperation Mode)
 			{
-				ToolSettings->SetBrushMode(Mode);
+				if (ToolSettings.IsValid())
+				{
+					ToolSettings->SetBrushMode(Mode);
+				}
 			})
 			+SSegmentedControl<EWeightEditOperation>::Slot(EWeightEditOperation::Add)
 			.Text(LOCTEXT("BrushAddMode", "Add"))
@@ -220,11 +229,14 @@ void FSkinWeightDetailCustomization::AddBrushUI(IDetailLayoutBuilder& DetailBuil
 					"Volume: falloff is based on the straight-line distance from the brush center to surrounding vertices.\n"))
 			.Value_Lambda([this]()
 			{
-				return ToolSettings->GetBrushConfig().FalloffMode;
+				return ToolSettings.IsValid() ? ToolSettings->GetBrushConfig().FalloffMode : EWeightBrushFalloffMode::Surface;
 			})
 			.OnValueChanged_Lambda([this](EWeightBrushFalloffMode Mode)
 			{
-				ToolSettings->SetFalloffMode(Mode);
+				if (ToolSettings.IsValid())
+				{
+					ToolSettings->SetFalloffMode(Mode);
+				}
 			})
 			+SSegmentedControl<EWeightBrushFalloffMode>::Slot(EWeightBrushFalloffMode::Surface)
 			.Text(LOCTEXT("SurfaceMode", "Surface"))
@@ -250,18 +262,24 @@ void FSkinWeightDetailCustomization::AddBrushUI(IDetailLayoutBuilder& DetailBuil
 		.SupportDynamicSliderMaxValue(true)
 		.Value_Lambda([this]()
 		{
-			return ToolSettings->GetBrushConfig().Radius;
+			return ToolSettings.IsValid() ? ToolSettings->GetBrushConfig().Radius : 20.f;
 		})
 		.OnValueChanged_Lambda([this](float NewValue)
 		{
-			ToolSettings->BrushRadius = NewValue;
-			ToolSettings->GetBrushConfig().Radius = NewValue;
-			FPropertyChangedEvent PropertyChangedEvent(UBrushBaseProperties::StaticClass()->FindPropertyByName(GET_MEMBER_NAME_CHECKED(UBrushBaseProperties, BrushRadius)));
-			ToolSettings->PostEditChangeProperty(PropertyChangedEvent);
+			if (ToolSettings.IsValid())
+			{
+				ToolSettings->BrushRadius = NewValue;
+				ToolSettings->GetBrushConfig().Radius = NewValue;
+				FPropertyChangedEvent PropertyChangedEvent(UBrushBaseProperties::StaticClass()->FindPropertyByName(GET_MEMBER_NAME_CHECKED(UBrushBaseProperties, BrushRadius)));
+				ToolSettings->PostEditChangeProperty(PropertyChangedEvent);
+			}
 		})
 		.OnValueCommitted_Lambda([this](float NewValue, ETextCommit::Type CommitType)
 		{
-			ToolSettings->SaveConfig();
+			if (ToolSettings.IsValid())
+			{
+				ToolSettings->SaveConfig();
+			}
 		})
 	];
 
@@ -283,18 +301,24 @@ void FSkinWeightDetailCustomization::AddBrushUI(IDetailLayoutBuilder& DetailBuil
 		.SupportDynamicSliderMaxValue(true)
 		.Value_Lambda([this]()
 		{
-			return ToolSettings->GetBrushConfig().Strength;
+			return ToolSettings.IsValid() ? ToolSettings->GetBrushConfig().Strength : 1.f;
 		})
 		.OnValueChanged_Lambda([this](float NewValue)
 		{
-			ToolSettings->BrushStrength = NewValue;
-			ToolSettings->GetBrushConfig().Strength = NewValue;
-			FPropertyChangedEvent PropertyChangedEvent(UBrushBaseProperties::StaticClass()->FindPropertyByName(GET_MEMBER_NAME_CHECKED(UBrushBaseProperties, BrushStrength)));
-			ToolSettings->PostEditChangeProperty(PropertyChangedEvent);
+			if (ToolSettings.IsValid())
+			{
+				ToolSettings->BrushStrength = NewValue;
+				ToolSettings->GetBrushConfig().Strength = NewValue;
+				FPropertyChangedEvent PropertyChangedEvent(UBrushBaseProperties::StaticClass()->FindPropertyByName(GET_MEMBER_NAME_CHECKED(UBrushBaseProperties, BrushStrength)));
+				ToolSettings->PostEditChangeProperty(PropertyChangedEvent);
+			}
 		})
 		.OnValueCommitted_Lambda([this](float NewValue, ETextCommit::Type CommitType)
 		{
-			ToolSettings->SaveConfig();
+			if (ToolSettings.IsValid())
+			{
+				ToolSettings->SaveConfig();
+			}
 		})
 	];
 
@@ -313,18 +337,24 @@ void FSkinWeightDetailCustomization::AddBrushUI(IDetailLayoutBuilder& DetailBuil
 		.MaxValue(1.f)
 		.Value_Lambda([this]()
 		{
-			return ToolSettings->GetBrushConfig().Falloff;
+			return ToolSettings.IsValid() ? ToolSettings->GetBrushConfig().Falloff : 1.f;
 		})
 		.OnValueChanged_Lambda([this](float NewValue)
 		{
-			ToolSettings->BrushFalloffAmount = NewValue;
-			ToolSettings->GetBrushConfig().Falloff = NewValue;
-			FPropertyChangedEvent PropertyChangedEvent(UBrushBaseProperties::StaticClass()->FindPropertyByName(GET_MEMBER_NAME_CHECKED(UBrushBaseProperties, BrushFalloffAmount)));
-			ToolSettings->PostEditChangeProperty(PropertyChangedEvent);
+			if (ToolSettings.IsValid())
+			{
+				ToolSettings->BrushFalloffAmount = NewValue;
+				ToolSettings->GetBrushConfig().Falloff = NewValue;
+				FPropertyChangedEvent PropertyChangedEvent(UBrushBaseProperties::StaticClass()->FindPropertyByName(GET_MEMBER_NAME_CHECKED(UBrushBaseProperties, BrushFalloffAmount)));
+				ToolSettings->PostEditChangeProperty(PropertyChangedEvent);
+			}
 		})
 		.OnValueCommitted_Lambda([this](float NewValue, ETextCommit::Type CommitType)
 		{
-			ToolSettings->SaveConfig();
+			if (ToolSettings.IsValid())
+			{
+				ToolSettings->SaveConfig();
+			}
 		})
 	];
 }
@@ -348,15 +378,18 @@ void FSkinWeightDetailCustomization::AddSelectionUI(IDetailLayoutBuilder& Detail
 		ToolbarBuilder.AddToolBarButton(FUIAction(
 		FExecuteAction::CreateLambda([this, Mode]()
 		{
-			ToolSettings->SetComponentMode(Mode);
+			if (ToolSettings.IsValid())
+			{
+				ToolSettings->SetComponentMode(Mode);
+			}
 		}),
 		FCanExecuteAction::CreateLambda([this]()
 		{
-			return ToolSettings->EditingMode == EWeightEditMode::Mesh;
+			return ToolSettings.IsValid() ? ToolSettings->EditingMode == EWeightEditMode::Mesh : false;
 		}),
 		FIsActionChecked::CreateLambda([this, Mode]()
 		{
-			return ToolSettings->ComponentSelectionMode == Mode;
+			return ToolSettings.IsValid() ? ToolSettings->ComponentSelectionMode == Mode : false;
 		})),
 		NAME_None,	// Extension hook
 		Label,		// Label
@@ -419,32 +452,43 @@ void FSkinWeightDetailCustomization::AddSelectionUI(IDetailLayoutBuilder& Detail
 							"NOTE: This only works on the target (main) mesh."))
 					.IsEnabled_Lambda([this]()
 					{
-						// isolated selection only available on main mesh (for now)
-						const bool bHasSelection = Tool->GetMainMeshSelector()->IsAnyComponentSelected();
-						const bool bAlreadyIsolatingSelection = Tool->GetSelectionIsolator()->IsSelectionIsolated();
-						return bHasSelection ||  bAlreadyIsolatingSelection;
+						if (Tool.IsValid())
+						{
+							// isolated selection only available on main mesh (for now)
+						   const bool bHasSelection = Tool->GetMainMeshSelector()->IsAnyComponentSelected();
+						   const bool bAlreadyIsolatingSelection = Tool->GetSelectionIsolator()->IsSelectionIsolated();
+						   return bHasSelection ||  bAlreadyIsolatingSelection;
+						}
+						return false;
 					})
 					.IsChecked_Lambda([this]()
 					{
-						const bool bAlreadyIsolatingSelection = Tool->GetSelectionIsolator()->IsSelectionIsolated();
-						return bAlreadyIsolatingSelection ? ECheckBoxState::Checked : ECheckBoxState::Unchecked;
+						if (Tool.IsValid())
+						{
+							const bool bAlreadyIsolatingSelection = Tool->GetSelectionIsolator()->IsSelectionIsolated();
+							return bAlreadyIsolatingSelection ? ECheckBoxState::Checked : ECheckBoxState::Unchecked;
+						}
+						return ECheckBoxState::Unchecked;
 					})
 					.OnCheckStateChanged_Lambda([this](ECheckBoxState InCheckBoxState)
 					{
-						if (InCheckBoxState == ECheckBoxState::Checked)
+						if (UWeightToolSelectionIsolator* SelectionIsolator = Tool.IsValid() ? Tool->GetSelectionIsolator() : nullptr)
 						{
-							Tool->GetSelectionIsolator()->IsolateSelectionAsTransaction();
-						}
-						else
-						{
-							Tool->GetSelectionIsolator()->UnIsolateSelectionAsTransaction();
+							if (InCheckBoxState == ECheckBoxState::Checked)
+							{
+								SelectionIsolator->IsolateSelectionAsTransaction();
+							}
+							else
+							{
+								SelectionIsolator->UnIsolateSelectionAsTransaction();
+							}
 						}
 					})
 					[
 						SNew(STextBlock)
 						.Text_Lambda([this]()
 						{
-							if (Tool->GetSelectionIsolator()->IsSelectionIsolated())
+							if (!Tool.IsValid() || Tool->GetSelectionIsolator()->IsSelectionIsolated())
 							{
 								return LOCTEXT("ShowAllButtonLabel", "Show Full Mesh");
 							}
@@ -465,7 +509,10 @@ void FSkinWeightDetailCustomization::AddSelectionUI(IDetailLayoutBuilder& Detail
 			.Padding(2.f, WeightEditVerticalPadding)
 			[
 				SNew(SButton)
-				.IsEnabled_Lambda([this]{ return Tool->GetActiveMeshSelector()->IsAnyComponentSelected(); })
+				.IsEnabled_Lambda([this]
+				{
+					return Tool.IsValid() ? Tool->GetActiveMeshSelector()->IsAnyComponentSelected() : false;
+				})
 				.HAlign(HAlign_Center)
 				.VAlign(VAlign_Center)
 				.Text(LOCTEXT("GrowSelectionButtonLabel", "Grow"))
@@ -473,7 +520,10 @@ void FSkinWeightDetailCustomization::AddSelectionUI(IDetailLayoutBuilder& Detail
 						"Grow the current selection by adding connected neighbors to current selection.\n"))
 				.OnClicked_Lambda([this]()
 				{
-					Tool->GetActiveMeshSelector()->GrowSelection();
+					if (Tool.IsValid())
+					{
+						Tool->GetActiveMeshSelector()->GrowSelection();
+					}
 					return FReply::Handled();
 				})
 			]
@@ -482,7 +532,10 @@ void FSkinWeightDetailCustomization::AddSelectionUI(IDetailLayoutBuilder& Detail
 			.Padding(2.f, WeightEditVerticalPadding)
 			[
 				SNew(SButton)
-				.IsEnabled_Lambda([this]{ return Tool->GetActiveMeshSelector()->IsAnyComponentSelected(); })
+				.IsEnabled_Lambda([this]
+				{
+					return Tool.IsValid() ? Tool->GetActiveMeshSelector()->IsAnyComponentSelected() : false;
+				})
 				.HAlign(HAlign_Center)
 				.VAlign(VAlign_Center)
 				.Text(LOCTEXT("ShrinkSelectionButtonLabel", "Shrink"))
@@ -490,7 +543,10 @@ void FSkinWeightDetailCustomization::AddSelectionUI(IDetailLayoutBuilder& Detail
 						"Shrink the current selection by removing components on the border of the current selection.\n"))
 				.OnClicked_Lambda([this]()
 				{
-					Tool->GetActiveMeshSelector()->ShrinkSelection();
+					if (Tool.IsValid())
+					{
+						Tool->GetActiveMeshSelector()->ShrinkSelection();
+					}
 					return FReply::Handled();
 				})
 			]
@@ -499,7 +555,10 @@ void FSkinWeightDetailCustomization::AddSelectionUI(IDetailLayoutBuilder& Detail
 			.Padding(2.f, WeightEditVerticalPadding)
 			[
 				SNew(SButton)
-				.IsEnabled_Lambda([this]{ return Tool->GetActiveMeshSelector()->IsAnyComponentSelected(); })
+				.IsEnabled_Lambda([this]
+				{
+					return Tool.IsValid() ? Tool->GetActiveMeshSelector()->IsAnyComponentSelected() : false;
+				})
 				.HAlign(HAlign_Center)
 				.VAlign(VAlign_Center)
 				.Text(LOCTEXT("FloodSelectionButtonLabel", "Flood"))
@@ -507,7 +566,10 @@ void FSkinWeightDetailCustomization::AddSelectionUI(IDetailLayoutBuilder& Detail
 						"Flood the current selection by adding all connected components to the current selection.\n"))
 				.OnClicked_Lambda([this]()
 				{
-					Tool->GetActiveMeshSelector()->FloodSelection();
+					if (Tool.IsValid())
+					{
+						Tool->GetActiveMeshSelector()->FloodSelection();
+					}
 					return FReply::Handled();
 				})
 			]
@@ -525,7 +587,7 @@ void FSkinWeightDetailCustomization::AddSelectionUI(IDetailLayoutBuilder& Detail
 				.IsEnabled_Lambda([this]
 				{
 					// only allow selecting affected vertices on the target/main mesh
-					return ToolSettings->MeshSelectMode == EMeshTransferOption::Target;
+					return ToolSettings.IsValid() ? ToolSettings->MeshSelectMode == EMeshTransferOption::Target : false;
 				})
 				.HAlign(HAlign_Center)
 				.VAlign(VAlign_Center)
@@ -535,7 +597,10 @@ void FSkinWeightDetailCustomization::AddSelectionUI(IDetailLayoutBuilder& Detail
 						"Holding Shift or Ctrl will add or subtract affected vertices from the current selection."))
 				.OnClicked_Lambda([this]()
 				{
-					Tool->SelectAffected();
+					if (Tool.IsValid())
+					{
+						Tool->SelectAffected();
+					}
 					return FReply::Handled();
 				})
 			]
@@ -544,14 +609,20 @@ void FSkinWeightDetailCustomization::AddSelectionUI(IDetailLayoutBuilder& Detail
 			.Padding(2.f, WeightEditVerticalPadding)
 			[
 				SNew(SButton)
-				.IsEnabled_Lambda([this]{ return Tool->GetActiveMeshSelector()->IsAnyComponentSelected(); })
+				.IsEnabled_Lambda([this]
+				{
+					return Tool.IsValid() ? Tool->GetActiveMeshSelector()->IsAnyComponentSelected() : false;
+				})
 				.HAlign(HAlign_Center)
 				.VAlign(VAlign_Center)
 				.Text(LOCTEXT("BorderSelectionButtonLabel", "Convert to Border"))
 				.ToolTipText(LOCTEXT("BorderSelectionTooltip", "Select vertices on the border of the current selection."))
 				.OnClicked_Lambda([this]()
 				{
-					Tool->GetActiveMeshSelector()->SelectBorder();
+					if (Tool.IsValid())
+					{
+						Tool->GetActiveMeshSelector()->SelectBorder();
+					}
 					return FReply::Handled();
 				})
 			]
@@ -569,7 +640,7 @@ void FSkinWeightDetailCustomization::AddSelectionUI(IDetailLayoutBuilder& Detail
 				.IsEnabled_Lambda([this]
 				{
 					// only allow selecting affected vertices on the target/main mesh
-					return ToolSettings->MeshSelectMode == EMeshTransferOption::Target;
+					return ToolSettings.IsValid() ? ToolSettings->MeshSelectMode == EMeshTransferOption::Target : false;
 				})
 				.HAlign(HAlign_Center)
 				.VAlign(VAlign_Center)
@@ -578,7 +649,10 @@ void FSkinWeightDetailCustomization::AddSelectionUI(IDetailLayoutBuilder& Detail
 						"Select vertices that are affected by at least N influences."))
 				.OnClicked_Lambda([this]()
 				{
-					Tool->SelectByInfluenceCount(Tool->GetWeightToolProperties()->ClampSelectValue);
+					if (Tool.IsValid())
+					{
+						Tool->SelectByInfluenceCount(Tool->GetWeightToolProperties()->ClampSelectValue);
+					}
 					return FReply::Handled();
 				})
 			]
@@ -591,15 +665,21 @@ void FSkinWeightDetailCustomization::AddSelectionUI(IDetailLayoutBuilder& Detail
 				.MaxValue(24)
 				.Value_Lambda([this]()
 				{
-					return ToolSettings->ClampSelectValue;
+					return ToolSettings.IsValid() ? ToolSettings->ClampSelectValue : 8;
 				})
 				.OnValueChanged_Lambda([this](int32 NewValue)
 				{
-					ToolSettings->ClampSelectValue = NewValue;
+					if (ToolSettings.IsValid())
+					{
+						ToolSettings->ClampSelectValue = NewValue;
+					}
 				})
 				.OnValueCommitted_Lambda([this](int32 NewValue, ETextCommit::Type CommitType)
 				{
-					ToolSettings->SaveConfig();
+					if (ToolSettings.IsValid())
+					{
+						ToolSettings->SaveConfig();
+					}
 				})
 			]
 		]
@@ -615,7 +695,7 @@ void FSkinWeightDetailCustomization::AddSelectionUI(IDetailLayoutBuilder& Detail
 		SNew(STextBlock)
 		.Visibility_Lambda([this]()
 		{
-			return !Tool->HasActiveSelectionOnMainMesh() ? EVisibility::Visible : EVisibility::Collapsed;
+			return Tool.IsValid() && !Tool->HasActiveSelectionOnMainMesh() ? EVisibility::Visible : EVisibility::Collapsed;
 		})
 		.Text_Lambda([this]()
 		{
@@ -629,7 +709,10 @@ void FSkinWeightDetailCustomization::AddSelectionUI(IDetailLayoutBuilder& Detail
 	.WholeRowContent()
 	[
 		SNew(SHorizontalBox)
-		.IsEnabled_Lambda([this]{ return Tool->HasActiveSelectionOnMainMesh(); })
+		.IsEnabled_Lambda([this]
+		{
+			return Tool.IsValid() && Tool->HasActiveSelectionOnMainMesh();
+		})
 
 		+SHorizontalBox::Slot()
 		[
@@ -640,12 +723,15 @@ void FSkinWeightDetailCustomization::AddSelectionUI(IDetailLayoutBuilder& Detail
 				"This operation applies interactively while dragging the slider. It operates on the currently selected vertices."))
 			.Value_Lambda([this]()
 			{
-				return ToolSettings->DirectEditState.EditMode;
+				return ToolSettings.IsValid() ? ToolSettings->DirectEditState.EditMode : EWeightEditOperation::Add;
 			})
 			.OnValueChanged_Lambda([this](EWeightEditOperation Mode)
 			{
-				ToolSettings->DirectEditState.EditMode = Mode;
-				ToolSettings->DirectEditState.Reset();
+				if(ToolSettings.IsValid())
+				{
+					ToolSettings->DirectEditState.EditMode = Mode;
+					ToolSettings->DirectEditState.Reset();
+				}
 			})
 			+SSegmentedControl<EWeightEditOperation>::Slot(EWeightEditOperation::Add)
 			.Text(LOCTEXT("BrushAddMode", "Add"))
@@ -664,94 +750,117 @@ void FSkinWeightDetailCustomization::AddSelectionUI(IDetailLayoutBuilder& Detail
 				SNew(SSpinBox<float>)
 				.Visibility_Lambda([this]()
 				{
-					const bool bIsVisible = ToolSettings->DirectEditState.EditMode != EWeightEditOperation::Relax;
+					const bool bIsVisible = ToolSettings.IsValid() && ToolSettings->DirectEditState.EditMode != EWeightEditOperation::Relax;
 					return bIsVisible ? EVisibility::Visible : EVisibility::Collapsed;
 				})
 				.MinSliderValue_Lambda([this]()
 				{
-					return ToolSettings->DirectEditState.GetModeMinValue();
+					return ToolSettings.IsValid() ? ToolSettings->DirectEditState.GetModeMinValue() : 0.f;
 				})
 				.MaxSliderValue_Lambda([this]()
 				{
-					return ToolSettings->DirectEditState.GetModeMaxValue();
+					return ToolSettings.IsValid() ? ToolSettings->DirectEditState.GetModeMaxValue() : 0.f;
 				})
 				.MinValue_Lambda([this]()
 				{
-					return ToolSettings->DirectEditState.GetModeMinValue();
+					return ToolSettings.IsValid() ? ToolSettings->DirectEditState.GetModeMinValue() : 0.f;
 				})
 				.MaxValue_Lambda([this]()
 				{
-					return ToolSettings->DirectEditState.GetModeMaxValue();
+					return ToolSettings.IsValid() ? ToolSettings->DirectEditState.GetModeMaxValue() : 0.f;
 				})
 				.Value_Lambda([this]()
 				{
-					return ToolSettings->DirectEditState.CurrentValue;
+					return ToolSettings.IsValid() ? ToolSettings->DirectEditState.CurrentValue : 0.f;
 				})
 				.OnValueChanged_Lambda([this](float NewValue)
 				{
-					ToolSettings->DirectEditState.CurrentValue = NewValue;
-					
-					if (ToolSettings->DirectEditState.bInTransaction)
+					if (ToolSettings.IsValid())
 					{
+						ToolSettings->DirectEditState.CurrentValue = NewValue;
+					
+					   if (ToolSettings->DirectEditState.bInTransaction)
+					   {
+						   float Value = NewValue;
+						   if (ToolSettings->DirectEditState.EditMode == EWeightEditOperation::Add)
+						   {
+							   Value = NewValue - ToolSettings->DirectEditState.StartValue;
+						   }
 						
-						float Value = NewValue;
-						if (ToolSettings->DirectEditState.EditMode == EWeightEditOperation::Add)
-						{
-							Value = NewValue - ToolSettings->DirectEditState.StartValue;
-						}
-						
-						constexpr bool bShouldTransact = false;
-						
-						Tool->EditWeightsOnVertices(
-							Tool->GetCurrentBoneIndex(),
-							Value,
-							0, // iterations
-							ToolSettings->DirectEditState.EditMode,
-							Tool->GetMainMeshSelector()->GetSelectedVertices(),
-							bShouldTransact);
+						   if (Tool.IsValid())
+						   {
+						   		constexpr bool bShouldTransact = false;
+					   			Tool->EditWeightsOnVertices(
+									  Tool->GetCurrentBoneIndex(),
+									  Value,
+									  0, // iterations
+									  ToolSettings->DirectEditState.EditMode,
+									  Tool->GetMainMeshSelector()->GetSelectedVertices(),
+									  bShouldTransact);
+						   }
+					   }
 					}
 				})
 				.OnValueCommitted_Lambda([this](float NewValue, ETextCommit::Type CommitType)
 				{
-					if (!ToolSettings->DirectEditState.bInTransaction)
+					if (ToolSettings.IsValid())
 					{
-						constexpr bool bShouldTransact = true;
-						Tool->EditWeightsOnVertices(
-							Tool->GetCurrentBoneIndex(),
-							NewValue,
-							0, // iterations
-							ToolSettings->DirectEditState.EditMode,
-							Tool->GetMainMeshSelector()->GetSelectedVertices(),
-							bShouldTransact);
+						if (!ToolSettings->DirectEditState.bInTransaction)
+						{
+							if (Tool.IsValid())
+							{
+								constexpr bool bShouldTransact = true;
+								Tool->EditWeightsOnVertices(
+									Tool->GetCurrentBoneIndex(),
+									NewValue,
+									0, // iterations
+									ToolSettings->DirectEditState.EditMode,
+									Tool->GetMainMeshSelector()->GetSelectedVertices(),
+									bShouldTransact);
+							}
+						}
+						ToolSettings->DirectEditState.bInTransaction = false;
 					}
-					ToolSettings->DirectEditState.bInTransaction = false;
 				})
 				.OnBeginSliderMovement_Lambda([this]()
 				{
-					ToolSettings->DirectEditState.StartValue = ToolSettings->DirectEditState.CurrentValue;
-					ToolSettings->DirectEditState.bInTransaction = true;
-					Tool->BeginChange();
+					if (ToolSettings.IsValid())
+					{
+						ToolSettings->DirectEditState.StartValue = ToolSettings->DirectEditState.CurrentValue;
+						ToolSettings->DirectEditState.bInTransaction = true;
+					}
+					if (Tool.IsValid())
+					{
+						Tool->BeginChange();
+					}
 				})
 				.OnEndSliderMovement_Lambda([this](float)
 				{
 					const FText TransactionLabel = LOCTEXT("FloodWeightChange", "Flood weights on vertices.");
-					Tool->EndChange(TransactionLabel);
-					ToolSettings->DirectEditState.bInTransaction = false;
-
-					// reset multiply slider
-					if (ToolSettings->DirectEditState.EditMode == EWeightEditOperation::Multiply)
+					if (Tool.IsValid())
 					{
-						// multiplying operation is always relative to 1.0
-						ToolSettings->DirectEditState.CurrentValue = 1.0f;
-						ToolSettings->DirectEditState.StartValue = 1.0f;
+						Tool->EndChange(TransactionLabel);
 					}
 
-					// reset add slider
-					if (ToolSettings->DirectEditState.EditMode == EWeightEditOperation::Add)
+					if (ToolSettings.IsValid())
 					{
-						// add operation is always relative to 0.0
-						ToolSettings->DirectEditState.CurrentValue = 0.0f;
-						ToolSettings->DirectEditState.StartValue = 0.0f;
+						ToolSettings->DirectEditState.bInTransaction = false;
+
+					   // reset multiply slider
+					   if (ToolSettings->DirectEditState.EditMode == EWeightEditOperation::Multiply)
+					   {
+						   // multiplying operation is always relative to 1.0
+						   ToolSettings->DirectEditState.CurrentValue = 1.0f;
+						   ToolSettings->DirectEditState.StartValue = 1.0f;
+					   }
+
+					   // reset add slider
+					   if (ToolSettings->DirectEditState.EditMode == EWeightEditOperation::Add)
+					   {
+						   // add operation is always relative to 0.0
+						   ToolSettings->DirectEditState.CurrentValue = 0.0f;
+						   ToolSettings->DirectEditState.StartValue = 0.0f;
+					   }
 					}
 				})
 				.ToolTipText(LOCTEXT("FloodWeightsToolTip", "Drag the slider to interactively adjust weights on the selected vertices."))
@@ -764,7 +873,10 @@ void FSkinWeightDetailCustomization::AddSelectionUI(IDetailLayoutBuilder& Detail
 	.WholeRowContent()
 	[
 		SNew(SVerticalBox)
-		.IsEnabled_Lambda([this]{ return Tool->HasActiveSelectionOnMainMesh(); })
+		.IsEnabled_Lambda([this]
+		{
+			return Tool.IsValid() ? Tool->HasActiveSelectionOnMainMesh() : false;
+		})
 
 		+ SVerticalBox::Slot()
 		.Padding(WeightEditHorizontalPadding, WeightEditVerticalPadding)
@@ -782,14 +894,17 @@ void FSkinWeightDetailCustomization::AddSelectionUI(IDetailLayoutBuilder& Detail
 					"This operation applies to the currently selected vertices."))
 					.OnClicked_Lambda([this]()
 					{
-						constexpr bool bShouldTransact = true;
-						Tool->EditWeightsOnVertices(
-							Tool->GetCurrentBoneIndex(),
-							ToolSettings->AddStrength,
-							0, // iterations 
-							EWeightEditOperation::Add,
-							Tool->GetMainMeshSelector()->GetSelectedVertices(),
-							bShouldTransact);
+						if (Tool.IsValid() && ToolSettings.IsValid())
+						{
+							constexpr bool bShouldTransact = true;
+							Tool->EditWeightsOnVertices(
+								Tool->GetCurrentBoneIndex(),
+								ToolSettings->AddStrength,
+								0, // iterations 
+								EWeightEditOperation::Add,
+								Tool->GetMainMeshSelector()->GetSelectedVertices(),
+								bShouldTransact);
+						}
 						return FReply::Handled();
 					})
 				]
@@ -803,15 +918,21 @@ void FSkinWeightDetailCustomization::AddSelectionUI(IDetailLayoutBuilder& Detail
 				.ToolTipText(LOCTEXT("AddWeightsSliderToolTip", "Adjust the value to Add to the selected vertices."))
 				.Value_Lambda([this]()
 				{
-					return ToolSettings->AddStrength;
+					return ToolSettings.IsValid() ? ToolSettings->AddStrength : 1.f;
 				})
 				.OnValueChanged_Lambda([this](float NewValue)
 				{
-					ToolSettings->AddStrength = NewValue;
+					if (ToolSettings.IsValid())
+					{
+						ToolSettings->AddStrength = NewValue;
+					}
 				})
 				.OnValueCommitted_Lambda([this](float NewValue, ETextCommit::Type CommitType)
 				{
-					ToolSettings->SaveConfig();
+					if (ToolSettings.IsValid())
+					{
+						ToolSettings->SaveConfig();
+					}
 				})
 			]
 		]
@@ -822,7 +943,10 @@ void FSkinWeightDetailCustomization::AddSelectionUI(IDetailLayoutBuilder& Detail
 	.WholeRowContent()
 	[
 		SNew(SVerticalBox)
-		.IsEnabled_Lambda([this]{ return Tool->HasActiveSelectionOnMainMesh(); })
+		.IsEnabled_Lambda([this]
+		{
+			return Tool.IsValid() ? Tool->HasActiveSelectionOnMainMesh() : false;
+		})
 
 		+ SVerticalBox::Slot()
 		.Padding(WeightEditHorizontalPadding, WeightEditVerticalPadding)
@@ -841,14 +965,17 @@ void FSkinWeightDetailCustomization::AddSelectionUI(IDetailLayoutBuilder& Detail
 						"This operation applies to the currently selected vertices."))
 					.OnClicked_Lambda([this]()
 					{
-						constexpr bool bShouldTransact = true;
-						Tool->EditWeightsOnVertices(
-							Tool->GetCurrentBoneIndex(),
-							ToolSettings->ReplaceValue,
-							0, // iterations
-							EWeightEditOperation::Replace,
-							Tool->GetMainMeshSelector()->GetSelectedVertices(),
-							bShouldTransact);
+						if (Tool.IsValid() && ToolSettings.IsValid())
+						{
+							constexpr bool bShouldTransact = true;
+							Tool->EditWeightsOnVertices(
+								Tool->GetCurrentBoneIndex(),
+								ToolSettings->ReplaceValue,
+								0, // iterations
+								EWeightEditOperation::Replace,
+								Tool->GetMainMeshSelector()->GetSelectedVertices(),
+								bShouldTransact);
+						}
 						return FReply::Handled();
 					})
 				]
@@ -862,15 +989,21 @@ void FSkinWeightDetailCustomization::AddSelectionUI(IDetailLayoutBuilder& Detail
 				.ToolTipText(LOCTEXT("ReplaceWeightsSliderToolTip", "Adjust the value to Replace on the selected vertices."))
 				.Value_Lambda([this]()
 				{
-					return ToolSettings->ReplaceValue;
+					return ToolSettings.IsValid() ? ToolSettings->ReplaceValue : 1.f;
 				})
 				.OnValueChanged_Lambda([this](float NewValue)
 				{
-					ToolSettings->ReplaceValue = NewValue;
+					if (ToolSettings.IsValid())
+					{
+						ToolSettings->ReplaceValue = NewValue;
+					}
 				})
 				.OnValueCommitted_Lambda([this](float NewValue, ETextCommit::Type CommitType)
 				{
-					ToolSettings->SaveConfig();
+					if (ToolSettings.IsValid())
+					{
+						ToolSettings->SaveConfig();
+					}
 				})
 			]
 		]
@@ -899,7 +1032,10 @@ void FSkinWeightDetailCustomization::AddSelectionUI(IDetailLayoutBuilder& Detail
 						"Weights on all selected vertices are set to the average of all selected vertices."))
 					.OnClicked_Lambda([this]()
 					{
-						ToolSettings->WeightTool->AverageWeights(ToolSettings->AverageStrength);
+						if (ToolSettings.IsValid())
+						{
+							ToolSettings->WeightTool->AverageWeights(ToolSettings->AverageStrength);
+						}
 						return FReply::Handled();
 					})
 				]
@@ -913,15 +1049,21 @@ void FSkinWeightDetailCustomization::AddSelectionUI(IDetailLayoutBuilder& Detail
 				.ToolTipText(LOCTEXT("AverageWeightsSliderToolTip", "Blend the amount to Average the weights on the selected vertices."))
 				.Value_Lambda([this]()
 				{
-					return ToolSettings->AverageStrength;
+					return ToolSettings.IsValid() ? ToolSettings->AverageStrength : 1.f;
 				})
 				.OnValueChanged_Lambda([this](float NewValue)
 				{
-					ToolSettings->AverageStrength = NewValue;
+					if (ToolSettings.IsValid())
+					{
+						ToolSettings->AverageStrength = NewValue;
+					}
 				})
 				.OnValueCommitted_Lambda([this](float NewValue, ETextCommit::Type CommitType)
 				{
-					ToolSettings->SaveConfig();
+					if (ToolSettings.IsValid())
+					{
+						ToolSettings->SaveConfig();
+					}
 				})
 			]
 		]
@@ -932,7 +1074,10 @@ void FSkinWeightDetailCustomization::AddSelectionUI(IDetailLayoutBuilder& Detail
 	.WholeRowContent()
 	[
 		SNew(SVerticalBox)
-		.IsEnabled_Lambda([this]{ return Tool->HasActiveSelectionOnMainMesh(); })
+		.IsEnabled_Lambda([this]
+		{
+			return Tool.IsValid() ? Tool->HasActiveSelectionOnMainMesh() : false;
+		})
 
 		+ SVerticalBox::Slot()
 		.Padding(WeightEditHorizontalPadding, WeightEditVerticalPadding)
@@ -950,15 +1095,18 @@ void FSkinWeightDetailCustomization::AddSelectionUI(IDetailLayoutBuilder& Detail
 						"This smooths weights across the mesh."))
 					.OnClicked_Lambda([this]()
 					{
-						constexpr bool bShouldTransact = true;
-						constexpr int32 DefaultRelaxIterations = 5; // provides a reasonable falloff distance 
-						Tool->EditWeightsOnVertices(
-							Tool->GetCurrentBoneIndex(),
-							ToolSettings->RelaxStrength,
-							DefaultRelaxIterations, 
-							EWeightEditOperation::Relax,
-							Tool->GetMainMeshSelector()->GetSelectedVertices(),
-							bShouldTransact);
+						if (Tool.IsValid() && ToolSettings.IsValid())
+						{
+							constexpr bool bShouldTransact = true;
+							constexpr int32 DefaultRelaxIterations = 5; // provides a reasonable falloff distance 
+							Tool->EditWeightsOnVertices(
+							   Tool->GetCurrentBoneIndex(),
+							   ToolSettings->RelaxStrength,
+							   DefaultRelaxIterations, 
+							   EWeightEditOperation::Relax,
+							   Tool->GetMainMeshSelector()->GetSelectedVertices(),
+							   bShouldTransact);
+						}
 						return FReply::Handled();
 					})
 				]
@@ -972,15 +1120,21 @@ void FSkinWeightDetailCustomization::AddSelectionUI(IDetailLayoutBuilder& Detail
 				.ToolTipText(LOCTEXT("RelaxWeightsSliderToolTip", "Blend the amount to Relax the weights on the selected vertices."))
 				.Value_Lambda([this]()
 				{
-					return ToolSettings->RelaxStrength;
+					return ToolSettings.IsValid() ? ToolSettings->RelaxStrength : 0.5f;
 				})
 				.OnValueChanged_Lambda([this](float NewValue)
 				{
-					ToolSettings->RelaxStrength = NewValue;
+					if (ToolSettings.IsValid())
+					{
+						ToolSettings->RelaxStrength = NewValue;
+					}
 				})
 				.OnValueCommitted_Lambda([this](float NewValue, ETextCommit::Type CommitType)
 				{
-					ToolSettings->SaveConfig();
+					if (ToolSettings.IsValid())
+					{
+						ToolSettings->SaveConfig();
+					}
 				})
 			]
 		]
@@ -1021,11 +1175,14 @@ void FSkinWeightDetailCustomization::AddSelectionUI(IDetailLayoutBuilder& Detail
 						"Z: copies weights across the XY plane."))
 					.Value_Lambda([this]()
 					{
-						return ToolSettings->MirrorAxis;
+						return ToolSettings.IsValid() ? ToolSettings->MirrorAxis : TEnumAsByte<EAxis::Type>(EAxis::X);
 					})
 					.OnValueChanged_Lambda([this](EAxis::Type Mode)
 					{
-						ToolSettings->MirrorAxis = Mode;
+						if (ToolSettings.IsValid())
+						{
+							ToolSettings->MirrorAxis = Mode;
+						}
 					})
 					+ SSegmentedControl<EAxis::Type>::Slot(EAxis::X)
 					.Text(LOCTEXT("MirrorXLabel", "X"))
@@ -1041,11 +1198,14 @@ void FSkinWeightDetailCustomization::AddSelectionUI(IDetailLayoutBuilder& Detail
 					.ToolTipText(LOCTEXT("MirrorDirectionTooltip", "The direction that determines what side of the plane to copy weights from."))
 					.Value_Lambda([this]()
 					{
-						return ToolSettings->MirrorDirection;
+						return ToolSettings.IsValid() ? ToolSettings->MirrorDirection : EMirrorDirection::PositiveToNegative;
 					})
 					.OnValueChanged_Lambda([this](EMirrorDirection Mode)
 					{
-						ToolSettings->MirrorDirection = Mode;
+						if (ToolSettings.IsValid())
+						{
+							ToolSettings->MirrorDirection = Mode;
+						}
 					})
 					+ SSegmentedControl<EMirrorDirection>::Slot(EMirrorDirection::PositiveToNegative)
 					.Text(LOCTEXT("MirrorPosToNegLabel", "+ to -"))
@@ -1068,7 +1228,10 @@ void FSkinWeightDetailCustomization::AddSelectionUI(IDetailLayoutBuilder& Detail
 					"This command operates on the selected vertices."))
 				.OnClicked_Lambda([this]()
 				{
-					Tool->MirrorWeights(ToolSettings->MirrorAxis, ToolSettings->MirrorDirection);
+					if (Tool.IsValid() && ToolSettings.IsValid())
+					{
+						Tool->MirrorWeights(ToolSettings->MirrorAxis, ToolSettings->MirrorDirection);
+					}
 					return FReply::Handled();
 				})
 			]
@@ -1094,11 +1257,14 @@ void FSkinWeightDetailCustomization::AddSelectionUI(IDetailLayoutBuilder& Detail
 					"This command operates on the selected vertices."))
 			.IsEnabled_Lambda([this]()
 			{
-				return ToolSettings->EditingMode == EWeightEditMode::Mesh;
+				return ToolSettings.IsValid() ? ToolSettings->EditingMode == EWeightEditMode::Mesh : false;
 			})
 			.OnClicked_Lambda([this]()
 			{
-				Tool->NormalizeWeights();
+				if (Tool.IsValid())
+				{
+					Tool->NormalizeWeights();
+				}
 				return FReply::Handled();
 			})
 		]
@@ -1115,11 +1281,14 @@ void FSkinWeightDetailCustomization::AddSelectionUI(IDetailLayoutBuilder& Detail
 					"This command operates on the selected vertices."))
 			.IsEnabled_Lambda([this]()
 			{
-				return ToolSettings->EditingMode == EWeightEditMode::Mesh;
+				return ToolSettings.IsValid() ? ToolSettings->EditingMode == EWeightEditMode::Mesh : false;
 			})
 			.OnClicked_Lambda([this]()
 			{
-				Tool->HammerWeights();
+				if (Tool.IsValid())
+				{
+					Tool->HammerWeights();
+				}
 				return FReply::Handled();
 			})
 		]
@@ -1145,7 +1314,10 @@ void FSkinWeightDetailCustomization::AddSelectionUI(IDetailLayoutBuilder& Detail
 				"This command operates on the selected vertices."))
 			.OnClicked_Lambda([this]()
 			{
-				ToolSettings->WeightTool->PruneWeights(ToolSettings->PruneValue, TArray<BoneIndex>());
+				if (ToolSettings.IsValid())
+				{
+					ToolSettings->WeightTool->PruneWeights(ToolSettings->PruneValue, TArray<BoneIndex>());
+				}
 				return FReply::Handled();
 			})
 		]
@@ -1158,15 +1330,21 @@ void FSkinWeightDetailCustomization::AddSelectionUI(IDetailLayoutBuilder& Detail
 			.MaxValue(1.f)
 			.Value_Lambda([this]()
 			{
-				return ToolSettings->PruneValue;
+				return ToolSettings.IsValid() ? ToolSettings->PruneValue : 0.01;
 			})
 			.OnValueChanged_Lambda([this](float NewValue)
 			{
-				ToolSettings->PruneValue = NewValue;
+				if (ToolSettings.IsValid())
+				{
+					ToolSettings->PruneValue = NewValue;
+				}
 			})
 			.OnValueCommitted_Lambda([this](float NewValue, ETextCommit::Type CommitType)
 			{
-				ToolSettings->SaveConfig();
+				if (ToolSettings.IsValid())
+				{
+					ToolSettings->SaveConfig();
+				}
 			})
 		]
 	];
@@ -1176,7 +1354,10 @@ void FSkinWeightDetailCustomization::AddSelectionUI(IDetailLayoutBuilder& Detail
 	.WholeRowContent()
 	[
 		SNew(SHorizontalBox)
-		.IsEnabled_Lambda([this]{ return Tool->HasActiveSelectionOnMainMesh(); })
+		.IsEnabled_Lambda([this]
+		{
+			return Tool.IsValid() ? Tool->HasActiveSelectionOnMainMesh() : false;
+		})
 
 		+SHorizontalBox::Slot()
 		.Padding(2.f, WeightEditVerticalPadding)
@@ -1190,7 +1371,10 @@ void FSkinWeightDetailCustomization::AddSelectionUI(IDetailLayoutBuilder& Detail
 				"This command operates on the selected vertices."))
 			.OnClicked_Lambda([this]()
 			{
-				ToolSettings->WeightTool->ClampInfluences(ToolSettings->ClampValue);
+				if (ToolSettings.IsValid())
+				{
+					ToolSettings->WeightTool->ClampInfluences(ToolSettings->ClampValue);
+				}
 				return FReply::Handled();
 			})
 		]
@@ -1203,15 +1387,21 @@ void FSkinWeightDetailCustomization::AddSelectionUI(IDetailLayoutBuilder& Detail
 			.MaxValue(24)
 			.Value_Lambda([this]()
 			{
-				return ToolSettings->ClampValue;
+				return ToolSettings.IsValid() ? ToolSettings->ClampValue : 8;
 			})
 			.OnValueChanged_Lambda([this](int32 NewValue)
 			{
-				ToolSettings->ClampValue = NewValue;
+				if (ToolSettings.IsValid())
+				{
+					ToolSettings->ClampValue = NewValue;
+				}
 			})
 			.OnValueCommitted_Lambda([this](int32 NewValue, ETextCommit::Type CommitType)
 			{
-				ToolSettings->SaveConfig();
+				if (ToolSettings.IsValid())
+				{
+					ToolSettings->SaveConfig();
+				}
 			})
 		]
 	];
@@ -1221,7 +1411,10 @@ void FSkinWeightDetailCustomization::AddSelectionUI(IDetailLayoutBuilder& Detail
 	.WholeRowContent()
 	[
 		SNew(SHorizontalBox)
-		.IsEnabled_Lambda([this]{ return Tool->HasActiveSelectionOnMainMesh(); })
+		.IsEnabled_Lambda([this]
+		{
+			return Tool.IsValid() ? Tool->HasActiveSelectionOnMainMesh() : false;
+		})
 
 		+SHorizontalBox::Slot()
 		.Padding(2.f, WeightEditVerticalPadding)
@@ -1236,7 +1429,10 @@ void FSkinWeightDetailCustomization::AddSelectionUI(IDetailLayoutBuilder& Detail
 					"This is designed to work with the Paste command."))
 				.OnClicked_Lambda([this]()
 				{
-					ToolSettings->WeightTool->CopyWeights();
+					if (ToolSettings.IsValid())
+					{
+						ToolSettings->WeightTool->CopyWeights();
+					}
 					return FReply::Handled();
 				})
 			]
@@ -1253,7 +1449,10 @@ void FSkinWeightDetailCustomization::AddSelectionUI(IDetailLayoutBuilder& Detail
 				"This command requires the clipboard contain weights from the Copy command."))
 			.OnClicked_Lambda([this]()
 			{
-				ToolSettings->WeightTool->PasteWeights();
+				if (ToolSettings.IsValid())
+				{
+					ToolSettings->WeightTool->PasteWeights();
+				}
 				return FReply::Handled();
 			})
 		]
@@ -1264,7 +1463,10 @@ void FSkinWeightDetailCustomization::AddSelectionUI(IDetailLayoutBuilder& Detail
 	.WholeRowContent()
 	[
 		SNew(SBox)
-		.IsEnabled_Lambda([this]{ return Tool->HasActiveSelectionOnMainMesh(); })
+		.IsEnabled_Lambda([this]
+		{
+			return Tool.IsValid() ? Tool->HasActiveSelectionOnMainMesh() : false;
+		})
 		[
 			SNew(SVertexWeightEditor, ToolSettings->WeightTool)
 		]
@@ -1297,12 +1499,15 @@ void FSkinWeightDetailCustomization::AddTransferUI(IDetailLayoutBuilder& DetailB
 						"If either mesh has no vertices selected, the whole mesh is considered.\n"))
 			.OnClicked_Lambda([this]()
 			{
-				Tool->GetWeightTransferManager()->TransferWeights();
+				if (Tool.IsValid())
+				{
+					Tool->GetWeightTransferManager()->TransferWeights();
+				}
 				return FReply::Handled();
 			})
 			.IsEnabled_Lambda([this]()
 			{
-				return Tool->GetWeightTransferManager()->CanTransferWeights();
+				return Tool.IsValid() ? Tool->GetWeightTransferManager()->CanTransferWeights() : false;
 			})
 		]
 	];
@@ -1341,12 +1546,15 @@ void FSkinWeightDetailCustomization::AddTransferUI(IDetailLayoutBuilder& DetailB
 					"Target: The mesh to copy weights TO (the main mesh in the tool)."))
 			.Value_Lambda([this]()
 			{
-				return ToolSettings->MeshSelectMode;
+				return ToolSettings.IsValid() ? ToolSettings->MeshSelectMode : EMeshTransferOption::Target;
 			})
 			.OnValueChanged_Lambda([this](EMeshTransferOption Mode)
 			{
-				ToolSettings->MeshSelectMode = Mode;
-				ToolSettings->WeightTool->UpdateSelectorState();
+				if (ToolSettings.IsValid())
+				{
+					ToolSettings->MeshSelectMode = Mode;
+					ToolSettings->WeightTool->UpdateSelectorState();
+				}
 			})
 			+ SSegmentedControl<EMeshTransferOption>::Slot(EMeshTransferOption::Source)
 			.Text(LOCTEXT("SourceMode", "Source"))
@@ -1371,7 +1579,7 @@ TSharedRef<SWidget> SVertexWeightItem::GenerateWidgetForColumn(const FName& Colu
 {
 	if (ColumnName == ColumnName_Bone)
 	{
-		const FName BoneName = ParentTable->Tool->GetBoneNameFromIndex(Element->BoneIndex);
+		const FName BoneName = ParentTable->Tool.IsValid() ? ParentTable->Tool->GetBoneNameFromIndex(Element->BoneIndex) : NAME_None;
 		return SNew(STextBlock).Text(FText::FromName(BoneName));
 	}
 
@@ -1389,6 +1597,11 @@ TSharedRef<SWidget> SVertexWeightItem::GenerateWidgetForColumn(const FName& Colu
 			{
 				return ValueDuringSlide;
 			}
+
+			if (!ParentTable->Tool.IsValid())
+			{
+				return 0.f;
+			}
 			
 			return ParentTable->Tool->GetAverageWeightOnBone(Element->BoneIndex, ParentTable->Tool->GetMainMeshSelector()->GetSelectedVertices());
 		})
@@ -1398,46 +1611,58 @@ TSharedRef<SWidget> SVertexWeightItem::GenerateWidgetForColumn(const FName& Colu
 			{
 				ValueDuringSlide = NewValue;
 
-				const bool bScalingUp = NewValue >= ValueAtStartOfSlide || FMath::IsNearlyEqual(ValueAtStartOfSlide, 0.f);
-				const float RangeEnd = bScalingUp ? 1.f : 0.f;
-				float RelativeScale = (NewValue - ValueAtStartOfSlide) * 1.f / (RangeEnd - ValueAtStartOfSlide);
-				RelativeScale *= bScalingUp ? 1.f : -1.f;
-				constexpr bool bShouldTransact = false;
-				ParentTable->Tool->EditWeightsOnVertices(
-					Element->BoneIndex,
-					RelativeScale,
-					0, // iterations
-					EWeightEditOperation::RelativeScale,
-					ParentTable->Tool->GetMainMeshSelector()->GetSelectedVertices(),
-					bShouldTransact);
+				if (ParentTable->Tool.IsValid())
+				{
+					const bool bScalingUp = NewValue >= ValueAtStartOfSlide || FMath::IsNearlyEqual(ValueAtStartOfSlide, 0.f);
+				   const float RangeEnd = bScalingUp ? 1.f : 0.f;
+				   float RelativeScale = (NewValue - ValueAtStartOfSlide) * 1.f / (RangeEnd - ValueAtStartOfSlide);
+				   RelativeScale *= bScalingUp ? 1.f : -1.f;
+				   constexpr bool bShouldTransact = false;
+				   ParentTable->Tool->EditWeightsOnVertices(
+					   Element->BoneIndex,
+					   RelativeScale,
+					   0, // iterations
+					   EWeightEditOperation::RelativeScale,
+					   ParentTable->Tool->GetMainMeshSelector()->GetSelectedVertices(),
+					   bShouldTransact);
+				}
 			}
 		})
 		.OnValueCommitted_Lambda([this](float NewValue, ETextCommit::Type CommitType)
 		{
 			if (!bInTransaction)
 			{
-				constexpr bool bShouldTransact = true;
-				ParentTable->Tool->EditWeightsOnVertices(
-					Element->BoneIndex,
-					NewValue,
-					0, // iterations,
-					EWeightEditOperation::Replace,
-					ParentTable->Tool->GetMainMeshSelector()->GetSelectedVertices(),
-					bShouldTransact);
+				if (ParentTable->Tool.IsValid())
+				{
+					constexpr bool bShouldTransact = true;
+					ParentTable->Tool->EditWeightsOnVertices(
+					   Element->BoneIndex,
+					   NewValue,
+					   0, // iterations,
+					   EWeightEditOperation::Replace,
+					   ParentTable->Tool->GetMainMeshSelector()->GetSelectedVertices(),
+					   bShouldTransact);
+				}
 			}
 			bInTransaction = false;
 		})
 		.OnBeginSliderMovement_Lambda([this]()
 		{
-			ParentTable->Tool->BeginChange();
-			ValueAtStartOfSlide = ParentTable->Tool->GetAverageWeightOnBone(Element->BoneIndex, ParentTable->Tool->GetMainMeshSelector()->GetSelectedVertices());
+			if (ParentTable->Tool.IsValid())
+			{
+				ParentTable->Tool->BeginChange();
+				ValueAtStartOfSlide = ParentTable->Tool->GetAverageWeightOnBone(Element->BoneIndex, ParentTable->Tool->GetMainMeshSelector()->GetSelectedVertices());
+			}
 			ValueDuringSlide = ValueAtStartOfSlide;
 			bInTransaction = true;
 		})
 		.OnEndSliderMovement_Lambda([this](float)
 		{
-			const FText TransactionLabel = LOCTEXT("DirectWeightChange", "Scale weights on vertices.");
-			ParentTable->Tool->EndChange(TransactionLabel);
+			if (ParentTable->Tool.IsValid())
+			{
+				const FText TransactionLabel = LOCTEXT("DirectWeightChange", "Scale weights on vertices.");
+				ParentTable->Tool->EndChange(TransactionLabel);
+			}
 			bInTransaction = false;
 		})
 		.ToolTipText(LOCTEXT("WeightSliderToolTip", "Set the weight on this bone for the selected vertices."));
@@ -1459,8 +1684,11 @@ TSharedRef<SWidget> SVertexWeightItem::GenerateWidgetForColumn(const FName& Colu
 			{
 				// use a negative threshold weight that no weight will ever
 				// be below because we only want to prune based on the bone, regardless of the weight
-				constexpr float NegativeThreshold = -1.0f;
-				ParentTable->Tool->PruneWeights(NegativeThreshold, {Element->BoneIndex});
+				if (ParentTable->Tool.IsValid())
+				{
+					constexpr float NegativeThreshold = -1.0f;
+					ParentTable->Tool->PruneWeights(NegativeThreshold, {Element->BoneIndex});
+				}
 				return FReply::Handled();
 			})
 			.Content()
@@ -1497,7 +1725,7 @@ void SVertexWeightEditor::Construct(const FArguments& InArgs, USkinWeightsPaintT
 			SAssignNew( ListView, SWeightEditorListViewType )
 			.Visibility_Lambda([this]()
 			{
-				return Tool->GetMainMeshSelector()->GetSelectedVertices().IsEmpty() ? EVisibility::Collapsed : EVisibility::Visible;
+				return Tool.IsValid() && !Tool->GetMainMeshSelector()->GetSelectedVertices().IsEmpty() ? EVisibility::Visible : EVisibility::Collapsed;
 			})
 			.SelectionMode(ESelectionMode::Single)
 			.ListItemsSource( &ListViewItems )
