[SKMModelingTools] Added support for accumulating tool changes into a DynamicMesh cache before apply them in bulk to the asset by when the Apply button is pressed.
+refactor morph target tool to be dynamic mesh based
+a new morph target manager widget

diff --git a/Engine/Plugins/Animation/SkeletalMeshModelingTools/Source/SkeletalMeshModelingTools/Private/Components/SKMBackedDynaMeshComponent.cpp b/Engine/Plugins/Animation/SkeletalMeshModelingTools/Source/SkeletalMeshModelingTools/Private/Components/SKMBackedDynaMeshComponent.cpp
new file mode 100644
--- /dev/null
+++ b/Engine/Plugins/Animation/SkeletalMeshModelingTools/Source/SkeletalMeshModelingTools/Private/Components/SKMBackedDynaMeshComponent.cpp
@@ -0,0 +1,663 @@
+﻿// Copyright Epic Games, Inc. All Rights Reserved.
+
+#include "Components/SKMBackedDynaMeshComponent.h"
+
+#include "DynamicMeshToMeshDescription.h"
+#include "LODUtilities.h"
+#include "MeshDescriptionToDynamicMesh.h"
+#include "ModelingToolTargetEditorOnlyUtil.h"
+#include "ModelingToolTargetUtil.h"
+#include "SkeletalMeshAttributes.h"
+#include "SkeletonModifier.h"
+#include "ScopedTransaction.h"
+#include "Misc/ITransaction.h"
+#include "SkeletalMeshEditorSubsystem.h"
+#include "ToolTargets/SkeletalMeshToolTarget.h"
+
+#include UE_INLINE_GENERATED_CPP_BY_NAME(SKMBackedDynaMeshComponent)
+
+#define LOCTEXT_NAMESPACE "SKMBackedDynaMeshComponent"
+
+void USkeletalMeshBackedDynamicMeshComponent::FSkeletonChangeTracker::Init(int32 NumBones)
+{
+	BoneIndexTracker.Reserve(NumBones);
+	for (int32 Index = 0; Index < NumBones; Index++)
+	{
+		BoneIndexTracker.Add(Index);
+	}
+
+	ChangeCount = 0;
+}
+
+void USkeletalMeshBackedDynamicMeshComponent::FSkeletonChangeTracker::HandleSkeletonChanged(const TArray<int32> ToolBoneIndexTracker)
+{
+	for (int32 BoneIndex = 0; BoneIndex < BoneIndexTracker.Num(); BoneIndex++)
+	{
+		int32 PreviousNewBoneIndex = BoneIndexTracker[BoneIndex];
+		int32 NewBoneIndex = ToolBoneIndexTracker[PreviousNewBoneIndex];
+		BoneIndexTracker[BoneIndex] = NewBoneIndex;
+	}
+
+	ChangeCount++;
+}
+
+int32 USkeletalMeshBackedDynamicMeshComponent::FSkeletonChangeTracker::GetChangeCount() const
+{
+	return ChangeCount;
+}
+
+const TArray<int32>& USkeletalMeshBackedDynamicMeshComponent::FSkeletonChangeTracker::GetBoneIndexTracker() const
+{
+	return BoneIndexTracker;
+}
+
+void USkeletalMeshBackedDynamicMeshComponent::FMorphTargetChangeTracker::Init(const TArray<FName> ExistingMorphTargets)
+{
+	OriginalNameToCurrentName.Reset();
+	CurrentNameToOriginalName.Reset();
+	for (const FName& MorphTarget : ExistingMorphTargets)
+	{
+		OriginalNameToCurrentName.Emplace(MorphTarget, MorphTarget);
+		CurrentNameToOriginalName.Emplace(MorphTarget, MorphTarget);
+	}
+}
+
+void USkeletalMeshBackedDynamicMeshComponent::FMorphTargetChangeTracker::HandleRenameMorphTarget(FName CurrentName, FName NewName)
+{
+	FName OriginalName = CurrentNameToOriginalName[CurrentName];
+
+	if (OriginalName != NAME_None)
+	{
+		OriginalNameToCurrentName[OriginalName] = NewName;
+	}
+
+	CurrentNameToOriginalName.Remove(CurrentName);
+	CurrentNameToOriginalName.Emplace(NewName, OriginalName);
+
+	if (EditedMorphTargets.Contains(CurrentName))
+	{
+		EditedMorphTargets.Remove(CurrentName);
+		EditedMorphTargets.Add(NewName);
+	}
+}
+
+void USkeletalMeshBackedDynamicMeshComponent::FMorphTargetChangeTracker::HandleRemoveMorphTarget(FName Name)
+{
+	FName OriginalName = CurrentNameToOriginalName[Name];
+	if (OriginalName != NAME_None)
+	{
+		OriginalNameToCurrentName[OriginalName] = NAME_None;
+	}
+
+	CurrentNameToOriginalName.Remove(Name);
+	
+	EditedMorphTargets.Remove(Name);
+}
+
+void USkeletalMeshBackedDynamicMeshComponent::FMorphTargetChangeTracker::HandleAddMorphTarget(FName Name)
+{
+	CurrentNameToOriginalName.Emplace(Name, NAME_None);
+}
+
+void USkeletalMeshBackedDynamicMeshComponent::FMorphTargetChangeTracker::HandleEditMorphTarget(FName Name)
+{
+	EditedMorphTargets.Add(Name);
+}
+
+FName USkeletalMeshBackedDynamicMeshComponent::FMorphTargetChangeTracker::GetCurrentMorphTargetName(FName OriginalName) const
+{
+	if (const FName* CurrentName = OriginalNameToCurrentName.Find(OriginalName))
+	{
+		return *CurrentName;
+	}
+
+	return NAME_None;
+}
+
+FName USkeletalMeshBackedDynamicMeshComponent::FMorphTargetChangeTracker::GetOriginalMorphTargetName(FName CurrentName) const
+{
+	if (const FName* OriginalName = CurrentNameToOriginalName.Find(CurrentName))
+	{
+		return *OriginalName;
+	}
+
+	return NAME_None;
+}
+
+const TSet<FName>& USkeletalMeshBackedDynamicMeshComponent::FMorphTargetChangeTracker::GetEditedMorphTargets() const
+{
+	return EditedMorphTargets;
+}
+
+TArray<USkeletalMeshBackedDynamicMeshComponent::FMorphTargetChangeTracker::FNameInfo> USkeletalMeshBackedDynamicMeshComponent::FMorphTargetChangeTracker::GetCurvesToRename() const
+{
+	TArray<FNameInfo> Renamed;
+	for (const TPair<FName, FName>& OriginalToCurrent : OriginalNameToCurrentName)
+	{
+		if (OriginalToCurrent.Value != OriginalToCurrent.Key && OriginalToCurrent.Value != NAME_None)
+		{
+			FNameInfo Info;
+			Info.OldName = OriginalToCurrent.Key;
+			Info.NewName = OriginalToCurrent.Value;
+			Renamed.Add(Info);
+		}
+	}
+
+	return Renamed;
+}
+
+TArray<FName> USkeletalMeshBackedDynamicMeshComponent::FMorphTargetChangeTracker::GetCurvesToRemove() const
+{
+	TArray<FName> Removed;
+	for (const TPair<FName, FName>& OriginalToCurrent : OriginalNameToCurrentName)
+	{
+		if (OriginalToCurrent.Value == NAME_None)
+		{
+			Removed.Add(OriginalToCurrent.Key);
+		}
+	}
+
+	return Removed;
+}
+
+TArray<FName> USkeletalMeshBackedDynamicMeshComponent::FMorphTargetChangeTracker::GetCurvesToAdd() const
+{
+	TArray<FName> Added;	
+	for (const TPair<FName, FName>& CurrentToOriginal : CurrentNameToOriginalName)
+	{
+		if (CurrentToOriginal.Value == NAME_None)
+		{
+			Added.Add(CurrentToOriginal.Key);
+		}
+	}
+	
+	return Added;
+}
+
+TArray<FName> USkeletalMeshBackedDynamicMeshComponent::FMorphTargetChangeTracker::GetCurrentMorphTargetNames() const
+{
+	TArray<FName> CurrentMorphTargetNames;
+	CurrentNameToOriginalName.GenerateKeyArray(CurrentMorphTargetNames);
+	return CurrentMorphTargetNames;
+}
+
+EMeshLODIdentifier USkeletalMeshBackedDynamicMeshComponent::Init(USkeletalMesh* InSkeletalMesh, EMeshLODIdentifier InLOD)
+{
+	WeakSkeletalMesh = InSkeletalMesh;
+	LOD = GetValidEditingLOD(InLOD);
+	LODIndex = static_cast<int32>(InLOD);
+
+	// Mesh
+	FMeshDescription* MeshDescription = InSkeletalMesh->GetMeshDescription(LODIndex);
+	FDynamicMesh3 LODMesh;
+	FMeshDescriptionToDynamicMesh Converter;
+	constexpr bool bWantTangents = true;
+	Converter.bVIDsFromNonManifoldMeshDescriptionAttr = true;
+	Converter.Convert(MeshDescription, LODMesh, bWantTangents);
+
+	SetMesh(MoveTemp(LODMesh));
+
+	// Additional Trackers
+	ResetTrackersDirect();
+
+	return static_cast<EMeshLODIdentifier>(LODIndex);
+}
+
+EMeshLODIdentifier USkeletalMeshBackedDynamicMeshComponent::GetLOD() const
+{
+	return LOD;
+}
+
+void USkeletalMeshBackedDynamicMeshComponent::CommitToSkeletalMesh()
+{
+	if (GetChangeCount() == 0)
+	{
+		return;
+	}
+
+	if (bIsCommiting)
+	{
+		return;
+	}
+
+	TGuardValue<bool> CommitGuard(bIsCommiting, true);
+	
+	USkeletalMesh* SkeletalMesh = WeakSkeletalMesh.Get();
+
+	TArray<FName> OldEngineGeneratedMorphTargets;
+	USkeletalMeshEditorSubsystem::GetMorphTargetsGeneratedByEngine(SkeletalMesh, OldEngineGeneratedMorphTargets);
+	
+	FScopedTransaction Transaction(LOCTEXT("ApplyChangesToAsset","Apply Changes To Skeletal Mesh"));
+
+	if (GetSkeletonChangeTracker().GetChangeCount() != 0)
+	{
+		USkeletonModifier* SkeletonModifier = NewObject<USkeletonModifier>();
+		SkeletonModifier->SetSkeletalMesh(SkeletalMesh);
+
+		SkeletonModifier->ExternalUpdate(GetRefSkeleton(), GetSkeletonChangeTracker().GetBoneIndexTracker());
+		SkeletonModifier->CommitSkeletonToSkeletalMesh();
+	}
+
+	USkeletalMeshToolTargetFactory* LocalFactory = NewObject<USkeletalMeshToolTargetFactory>();
+	LocalFactory->SetActiveEditingLOD(GetLOD());
+	USkeletalMeshToolTarget* LocalTarget = CastChecked<USkeletalMeshToolTarget>(LocalFactory->BuildTarget(SkeletalMesh, {}));
+
+	LocalTarget->CommitDynamicMesh(*GetMesh());
+
+	// Make sure morph targets are marked as engine generated correctly
+	// As soon as a morph is edited in engine, we want to mark it such that reimports in the future don't overwrite our edits.
+	{
+		TArray<FString> NewEngineGeneratedMorphTargets;
+	
+		for (const FName& OldEngineGeneratedMorphTarget : OldEngineGeneratedMorphTargets)
+		{
+			FName CurrentName = GetMorphTargetChangeTracker().GetCurrentMorphTargetName(OldEngineGeneratedMorphTarget);
+			if (CurrentName != NAME_None)
+			{
+				NewEngineGeneratedMorphTargets.AddUnique(CurrentName.ToString());
+				USkeletalMeshEditorSubsystem::SetMorphTargetsToGeneratedByEngine(SkeletalMesh, {CurrentName.ToString()});
+			}
+		}
+	
+		for (const FName& EditedMorphTarget : GetMorphTargetChangeTracker().GetEditedMorphTargets())
+		{
+			NewEngineGeneratedMorphTargets.AddUnique(EditedMorphTarget.ToString());	
+		}
+	
+		USkeletalMeshEditorSubsystem::SetMorphTargetsToGeneratedByEngine(SkeletalMesh, NewEngineGeneratedMorphTargets);
+	}
+
+	// Also change related curves
+	{
+		
+		UAnimCurveMetaData* AnimCurveMetaData = SkeletalMesh->GetAssetUserData<UAnimCurveMetaData>();
+		if (AnimCurveMetaData == nullptr)
+		{
+			AnimCurveMetaData = NewObject<UAnimCurveMetaData>(SkeletalMesh, NAME_None, RF_Transactional);
+			SkeletalMesh->AddAssetUserData(AnimCurveMetaData);
+		}
+
+		for (const FMorphTargetChangeTracker::FNameInfo& NameInfo : GetMorphTargetChangeTracker().GetCurvesToRename())
+		{
+			// Only rename if we have a morph flag set
+			if (FCurveMetaData* CurveMetaData = AnimCurveMetaData->GetCurveMetaData(NameInfo.OldName))
+			{
+				if (CurveMetaData->Type.bMorphtarget)
+				{
+					AnimCurveMetaData->RenameCurveMetaData(NameInfo.OldName, NameInfo.NewName);
+				}
+			}	
+		}
+	
+		for (const FName OldCurve : GetMorphTargetChangeTracker().GetCurvesToRemove())
+		{
+			// Only remove if we have a morph flag set
+			if (FCurveMetaData* CurveMetaData = AnimCurveMetaData->GetCurveMetaData(OldCurve))
+			{
+				if (CurveMetaData->Type.bMorphtarget)
+				{
+					AnimCurveMetaData->RemoveCurveMetaData(OldCurve);
+				}
+			}	
+		}
+	
+		for (const FName NewCurve : GetMorphTargetChangeTracker().GetCurvesToAdd())
+		{
+			AnimCurveMetaData->AddCurveMetaData(NewCurve);
+		
+			// Ensure we have a morph flag set
+			if (FCurveMetaData* CurveMetaData = AnimCurveMetaData->GetCurveMetaData(NewCurve))
+			{
+				AnimCurveMetaData->Modify();
+				CurveMetaData->Type.bMorphtarget = true;
+			}	
+		}
+	}
+
+	FTrackerChangeScope(this);
+	ResetTrackersDirect();
+	
+
+}
+
+bool USkeletalMeshBackedDynamicMeshComponent::IsCommitting() const
+{
+	return bIsCommiting;
+}
+
+void USkeletalMeshBackedDynamicMeshComponent::DiscardChanges()
+{
+	// Records a change that turns current mesh back to the initial state
+	FChangeScope(this);
+	
+	FDynamicMesh3 InitialMeshCopy = *Trackers.InitialAssetMesh;
+	SetMesh(MoveTemp(InitialMeshCopy));
+
+	ResetTrackersDirect();
+}
+
+
+void USkeletalMeshBackedDynamicMeshComponent::HandleSkeletonChange(USkeletonModifier* InModifier)
+{
+	FTrackerChangeScope ChangeScope(this);
+	
+	Trackers.RefSkeleton = InModifier->GetReferenceSkeleton();
+	Trackers.SkeletonChangeTracker.HandleSkeletonChanged(InModifier->GetBoneIndexTracker());
+	MarkDirtyInternal();
+}
+
+void USkeletalMeshBackedDynamicMeshComponent::ForwardVisibilityChangeRequest(bool bInVisible)
+{
+	// This component is typically hidden as it only serves as data container
+	// A separate mesh should be used to preview the skeletal mesh represented by this component,
+	// which should receive visibility updates instead.
+	OnRequestingVisibilityChangeDelegate.Broadcast(bInVisible);
+}
+
+FName USkeletalMeshBackedDynamicMeshComponent::AddMorphTarget(FName InName)
+{
+	FName ActualName = GetAvailableMorphTargetName(InName);
+	
+	FChangeScope ChangeScope(this);
+
+	// Mesh change
+	EditMesh([&](FDynamicMesh3& Mesh)
+	{
+		using namespace UE::Geometry;
+		FDynamicMeshMorphTargetAttribute* MorphTargetAttribute = new FDynamicMeshMorphTargetAttribute(&Mesh);
+		MorphTargetAttribute->SetName(ActualName);
+		Mesh.Attributes()->AttachMorphTargetAttribute(ActualName, MorphTargetAttribute);
+	});
+
+	// Tracker change
+	Trackers.MorphTargetChangeTracker.HandleAddMorphTarget(ActualName);
+	MarkDirtyInternal();
+
+	return ActualName;
+}
+
+FName USkeletalMeshBackedDynamicMeshComponent::RenameMorphTarget(FName InOldName, FName InNewName)
+{
+	if (InOldName == InNewName)
+	{
+		return InNewName;
+	}
+	
+	if (!GetMorphTargetChangeTracker().GetCurrentMorphTargetNames().Contains(InOldName))
+	{
+		return NAME_None;
+	}
+
+	
+	FName ActualNewName = GetAvailableMorphTargetName(InNewName, InOldName);
+
+	FChangeScope ChangeScope(this);
+
+	EditMesh([&](FDynamicMesh3& Mesh)
+	{
+		using namespace UE::Geometry;
+		FDynamicMeshMorphTargetAttribute* OldMorphTargetAttribute = Mesh.Attributes()->GetMorphTargetAttribute(InOldName);
+			
+		FDynamicMeshMorphTargetAttribute* NewMorphTargetAttribute = new FDynamicMeshMorphTargetAttribute(&Mesh);
+		NewMorphTargetAttribute->SetName(ActualNewName);
+		*NewMorphTargetAttribute = MoveTemp(*OldMorphTargetAttribute);
+			
+		Mesh.Attributes()->AttachMorphTargetAttribute(ActualNewName, NewMorphTargetAttribute);
+		Mesh.Attributes()->RemoveMorphTargetAttribute(InOldName);
+	});
+
+	Trackers.MorphTargetChangeTracker.HandleRenameMorphTarget(InOldName, ActualNewName);
+	MarkDirtyInternal();
+
+	return ActualNewName;
+}
+
+void USkeletalMeshBackedDynamicMeshComponent::RemoveMorphTargets(const TArray<FName>& InNames)
+{
+	TArray<FName> MorphsToRemove;
+	for (const FName& Name : InNames)
+	{
+		if (GetMorphTargetChangeTracker().GetCurrentMorphTargetNames().Contains(Name))
+		{
+			MorphsToRemove.Add(Name);
+		}
+	}
+	
+	if (MorphsToRemove.IsEmpty())
+	{
+		return;
+	}
+
+	FChangeScope ChangeScope(this);
+	
+	EditMesh([&](FDynamicMesh3& Mesh)
+	{
+		using namespace UE::Geometry;
+		for (const FName& Name : MorphsToRemove)
+		{
+			Mesh.Attributes()->RemoveMorphTargetAttribute(Name);
+		}
+	});
+
+	for (const FName& Name : MorphsToRemove)
+	{
+		Trackers.MorphTargetChangeTracker.HandleRemoveMorphTarget(Name);
+	}
+
+	MarkDirtyInternal();
+}
+
+TArray<FName> USkeletalMeshBackedDynamicMeshComponent::DuplicateMorphTargets(const TArray<FName>& InNames)
+{
+	TArray<FName> MorphsToDuplicate;
+	for (const FName& Name : InNames)
+	{
+		if (GetMorphTargetChangeTracker().GetCurrentMorphTargetNames().Contains(Name))
+		{
+			MorphsToDuplicate.Add(Name);
+		}
+	}
+	
+	if (MorphsToDuplicate.IsEmpty())
+	{
+		return {};
+	}
+
+	FChangeScope ChangeScope(this);
+
+	TArray<FName> NewMorphs;
+	EditMesh([&](FDynamicMesh3& Mesh)
+	{
+		using namespace UE::Geometry;
+		for (const FName& Name : MorphsToDuplicate)
+		{
+			FName DuplicateName = GetAvailableMorphTargetName(Name);
+			NewMorphs.Add(DuplicateName);
+			
+			using namespace UE::Geometry;
+			FDynamicMeshMorphTargetAttribute* SourceMorphTargetAttribute = Mesh.Attributes()->GetMorphTargetAttribute(Name);
+			
+			FDynamicMeshMorphTargetAttribute* DuplicatedMorphTargetAttribute = new FDynamicMeshMorphTargetAttribute(&Mesh);
+			DuplicatedMorphTargetAttribute->SetName(DuplicateName);
+			*DuplicatedMorphTargetAttribute = (*SourceMorphTargetAttribute);
+			
+			Mesh.Attributes()->AttachMorphTargetAttribute(DuplicateName, DuplicatedMorphTargetAttribute);
+		}
+	});
+	
+	// Tracker change
+	for (const FName& Name : NewMorphs)
+	{
+		Trackers.MorphTargetChangeTracker.HandleAddMorphTarget(Name);
+		Trackers.MorphTargetChangeTracker.HandleEditMorphTarget(Name);
+	}
+
+	
+	return NewMorphs;
+}
+
+void USkeletalMeshBackedDynamicMeshComponent::MarkMorphTargetEdited(FName InName)
+{
+	FTrackerChangeScope Scope(this);
+	
+	// Mesh should be marked dirty during tool commit already 
+	Trackers.MorphTargetChangeTracker.HandleEditMorphTarget(InName);
+
+}
+
+bool USkeletalMeshBackedDynamicMeshComponent::IsSkeletonDirty() const
+{
+	return GetSkeletonChangeTracker().GetChangeCount() != 0;
+}
+
+bool USkeletalMeshBackedDynamicMeshComponent::IsDirty() const
+{
+	return GetChangeCount() != 0;
+}
+
+int32 USkeletalMeshBackedDynamicMeshComponent::GetChangeCount() const
+{
+	return Trackers.ChangeCount;
+}
+
+USkeletalMesh* USkeletalMeshBackedDynamicMeshComponent::GetSkeletalMesh() const
+{
+	return WeakSkeletalMesh.Get();
+}
+
+void USkeletalMeshBackedDynamicMeshComponent::ResetTrackersDirect()
+{
+	Trackers.InitialAssetMesh = GetMeshCopy();
+
+	Trackers.ChangeCount = 0;
+	
+	Trackers.RefSkeleton = WeakSkeletalMesh->GetRefSkeleton();
+
+	Trackers.SkeletonChangeTracker.Init(Trackers.RefSkeleton.GetRawBoneNum());
+	
+	TArray<FName> MorphTargetNames;
+	GetMesh()->Attributes()->GetMorphTargetAttributes().GenerateKeyArray(MorphTargetNames);
+	Trackers.MorphTargetChangeTracker.Init(MorphTargetNames);
+}
+
+FName USkeletalMeshBackedDynamicMeshComponent::GetAvailableMorphTargetName(FName InNewName, FName InOldName) const
+{
+	TArray<FName> MorphTargets = GetMorphTargetChangeTracker().GetCurrentMorphTargetNames();
+
+	MorphTargets.Remove(InOldName);
+	
+	FName ActualName = InNewName;	
+	while (MorphTargets.Contains(ActualName))
+	{
+		ActualName.SetNumber(ActualName.GetNumber() + 1);
+	}
+
+	return ActualName;
+}
+
+EMeshLODIdentifier USkeletalMeshBackedDynamicMeshComponent::GetValidEditingLOD(EMeshLODIdentifier InLOD) const
+{
+	constexpr bool bSkipAutoGenerated= true;
+	TArray<EMeshLODIdentifier> AssetAvailableLODs = UE::ToolTarget::GetAvailableLODs(WeakSkeletalMesh.Get(), bSkipAutoGenerated);
+
+	if (!AssetAvailableLODs.Contains(InLOD))
+	{
+		return AssetAvailableLODs[0];
+	}
+
+	return InLOD;
+}
+
+TSharedPtr<FDynamicMesh3> USkeletalMeshBackedDynamicMeshComponent::GetMeshCopy()
+{
+	TSharedPtr<FDynamicMesh3> Mesh = MakeShared<FDynamicMesh3>();
+	ProcessMesh([&](const FDynamicMesh3& ReadMesh) { *Mesh = ReadMesh; });
+	return Mesh;	
+}
+
+void USkeletalMeshBackedDynamicMeshComponent::MarkDirtyInternal()
+{
+	Trackers.ChangeCount++;
+}
+
+
+void USkeletalMeshBackedDynamicMeshComponent::MarkDirty()
+{
+	FTrackerChangeScope ChangeScope(this);
+
+	MarkDirtyInternal();
+}
+
+
+
+USkeletalMeshBackedDynamicMeshComponent::FTrackerChange::FTrackerChange(USkeletalMeshBackedDynamicMeshComponent* InComponent)
+{
+	OldTrackers = InComponent->Trackers;
+	NewTrackers = OldTrackers; 
+}
+
+FString USkeletalMeshBackedDynamicMeshComponent::FTrackerChange::ToString() const
+{
+	return TEXT("State Changed");
+}
+
+void USkeletalMeshBackedDynamicMeshComponent::FTrackerChange::Apply(UObject* Object)
+{
+	if (USkeletalMeshBackedDynamicMeshComponent* Component = Cast<USkeletalMeshBackedDynamicMeshComponent>(Object))
+	{
+		Component->Trackers = NewTrackers;
+	}	
+}
+
+void USkeletalMeshBackedDynamicMeshComponent::FTrackerChange::Revert(UObject* Object)
+{
+	if (USkeletalMeshBackedDynamicMeshComponent* Component = Cast<USkeletalMeshBackedDynamicMeshComponent>(Object))
+	{
+		Component->Trackers = OldTrackers;
+	}	
+}
+
+void USkeletalMeshBackedDynamicMeshComponent::FTrackerChange::Close(USkeletalMeshBackedDynamicMeshComponent* InComponent)
+{
+	NewTrackers = InComponent->Trackers;	
+}
+
+USkeletalMeshBackedDynamicMeshComponent::FChangeScope::FChangeScope(USkeletalMeshBackedDynamicMeshComponent* InComponent, bool bInRecordMeshChange)
+{
+	Component = InComponent;
+	bRecordMeshChange = bInRecordMeshChange;
+	
+	if (GUndo)
+	{
+		TrackerChange = MakeUnique<FTrackerChange>(Component);
+
+		if (bRecordMeshChange)
+		{
+			OldMesh = Component->GetMeshCopy();
+			NewMesh = OldMesh;
+		}
+	}
+	
+}
+
+USkeletalMeshBackedDynamicMeshComponent::FChangeScope::~FChangeScope()
+{
+	if (GUndo)
+	{
+		TrackerChange->Close(Component);
+
+		GUndo->StoreUndo(Component, MoveTemp(TrackerChange));	
+		
+		if (bRecordMeshChange)
+		{
+			NewMesh = Component->GetMeshCopy();
+
+			TUniquePtr<FMeshReplacementChange> ReplaceChange = MakeUnique<FMeshReplacementChange>(OldMesh, NewMesh);
+
+			Component->GetDynamicMesh()->Modify();
+			GUndo->StoreUndo(Component->GetDynamicMesh(), MoveTemp(ReplaceChange));	
+		}
+	}
+}
+
+#undef LOCTEXT_NAMESPACE
diff --git a/Engine/Plugins/Animation/SkeletalMeshModelingTools/Source/SkeletalMeshModelingTools/Private/Components/SKMBackedDynaMeshComponent.h b/Engine/Plugins/Animation/SkeletalMeshModelingTools/Source/SkeletalMeshModelingTools/Private/Components/SKMBackedDynaMeshComponent.h
new file mode 100644
--- /dev/null
+++ b/Engine/Plugins/Animation/SkeletalMeshModelingTools/Source/SkeletalMeshModelingTools/Private/Components/SKMBackedDynaMeshComponent.h
@@ -0,0 +1,170 @@
+﻿// Copyright Epic Games, Inc. All Rights Reserved.
+
+#pragma once
+
+#include "CoreMinimal.h"
+#include "Components/DynamicMeshComponent.h"
+#include "ReferenceSkeleton.h"
+#include "TargetInterfaces/MeshTargetInterfaceTypes.h"
+
+#include "SKMBackedDynaMeshComponent.generated.h"
+
+#define UE_API SKELETALMESHMODELINGTOOLS_API
+
+
+class USkeletonModifier;
+
+
+UCLASS(MinimalAPI, Transient)
+class USkeletalMeshBackedDynamicMeshComponent :
+	public UDynamicMeshComponent
+{
+	GENERATED_BODY()
+public:	
+	
+	struct FSkeletonChangeTracker
+	{
+		void Init(int32 NumBones);
+		void HandleSkeletonChanged(const TArray<int32> ToolBoneIndexTracker);
+		int32 GetChangeCount() const;
+		const TArray<int32>& GetBoneIndexTracker() const;
+
+	private:
+		// Used to update retarget settings
+		TArray<int32> BoneIndexTracker;
+
+		int32 ChangeCount = 0;
+	};
+
+	struct FMorphTargetChangeTracker
+	{
+		void Init(const TArray<FName> ExistingMorphTargets);
+		void HandleRenameMorphTarget(FName CurrentName, FName NewName);
+		void HandleRemoveMorphTarget(FName Name);
+		void HandleAddMorphTarget(FName Name);
+		void HandleEditMorphTarget(FName Name);
+		FName GetCurrentMorphTargetName(FName OriginalName) const;
+		FName GetOriginalMorphTargetName(FName CurrentName) const;
+		const TSet<FName>& GetEditedMorphTargets() const;
+		struct FNameInfo
+		{
+			FName OldName;
+			FName NewName;
+		};
+		TArray<FNameInfo> GetCurvesToRename() const;
+		TArray<FName> GetCurvesToRemove() const;
+		TArray<FName> GetCurvesToAdd() const;
+
+		TArray<FName> GetCurrentMorphTargetNames() const;
+	
+	private:
+		TMap<FName, FName> OriginalNameToCurrentName;
+		TMap<FName, FName> CurrentNameToOriginalName;
+		TSet<FName> EditedMorphTargets;
+
+	};
+	
+
+	EMeshLODIdentifier Init(USkeletalMesh* InSkeletalMesh, EMeshLODIdentifier InLOD);
+	EMeshLODIdentifier GetLOD() const;
+	void CommitToSkeletalMesh();
+	bool IsCommitting() const;
+	void DiscardChanges();
+
+	void MarkDirty();
+	void HandleSkeletonChange(USkeletonModifier* InModifier);
+	void ForwardVisibilityChangeRequest(bool bInVisible);
+
+	FName AddMorphTarget(FName InName);
+	FName RenameMorphTarget(FName InOldName, FName InNewName);
+	void RemoveMorphTargets(const TArray<FName>& InNames);
+	TArray<FName> DuplicateMorphTargets(const TArray<FName>& InNames);
+	
+	void MarkMorphTargetEdited(FName InName);
+	
+	const FSkeletonChangeTracker& GetSkeletonChangeTracker() const {return Trackers.SkeletonChangeTracker;}
+	const FMorphTargetChangeTracker& GetMorphTargetChangeTracker() const {return Trackers.MorphTargetChangeTracker;}
+	const FReferenceSkeleton& GetRefSkeleton() const {return Trackers.RefSkeleton;}
+
+	bool IsSkeletonDirty() const;
+	bool IsDirty() const;
+	int32 GetChangeCount() const;
+	
+
+	USkeletalMesh* GetSkeletalMesh() const;
+	
+	DECLARE_MULTICAST_DELEGATE_OneParam(FOnRequestingVisibilityChange, bool);
+	FOnRequestingVisibilityChange& GetOnRequestingVisibilityChange() {return OnRequestingVisibilityChangeDelegate;}
+	
+protected:
+	void ResetTrackersDirect();
+	FName GetAvailableMorphTargetName(FName InNewName, FName InOldName = NAME_None) const;
+	EMeshLODIdentifier GetValidEditingLOD(EMeshLODIdentifier InLOD) const;
+
+	TSharedPtr<FDynamicMesh3> GetMeshCopy();
+	void MarkDirtyInternal();
+
+	struct FTrackers
+	{
+		TSharedPtr<FDynamicMesh3> InitialAssetMesh;
+		
+		int32 ChangeCount = 0;
+		
+		FReferenceSkeleton RefSkeleton;
+		FSkeletonChangeTracker SkeletonChangeTracker;
+		FMorphTargetChangeTracker MorphTargetChangeTracker;
+	};
+
+	FTrackers Trackers;
+
+	TWeakObjectPtr<USkeletalMesh> WeakSkeletalMesh;
+	EMeshLODIdentifier LOD = EMeshLODIdentifier::LOD0;
+	int32 LODIndex = 0;
+
+	bool bIsCommiting = false;
+	
+	FOnRequestingVisibilityChange OnRequestingVisibilityChangeDelegate;
+	
+	class FTrackerChange : public FToolCommandChange
+	{
+	public:
+		FTrackers OldTrackers;
+		FTrackers NewTrackers;
+		FTrackerChange(USkeletalMeshBackedDynamicMeshComponent* InComponent);
+		virtual FString ToString() const override;
+		virtual void Apply(UObject* Object) override;
+		virtual void Revert(UObject* Object) override;
+		void Close(USkeletalMeshBackedDynamicMeshComponent* InComponent);	
+	};
+	
+	// Records into current Transaction if there is one
+	class FChangeScope 
+	{
+	public:
+		FChangeScope(USkeletalMeshBackedDynamicMeshComponent* InComponent, bool bInRecordMeshChange = true);
+		~FChangeScope();
+		
+	private:
+		TUniquePtr<FTrackerChange> TrackerChange;
+		TSharedPtr<FDynamicMesh3> OldMesh;
+		TSharedPtr<FDynamicMesh3> NewMesh;
+
+		USkeletalMeshBackedDynamicMeshComponent* Component = nullptr;
+		bool bRecordMeshChange = false;
+
+		friend class FTrackerChange;
+	};
+
+	class FTrackerChangeScope : public FChangeScope
+	{
+	public:
+		FTrackerChangeScope(USkeletalMeshBackedDynamicMeshComponent* InComponent) : FChangeScope(InComponent, false) {}
+	};
+
+
+	
+	friend class FChangeScope;
+	friend class FTrackerChange;
+};
+
+#undef UE_API
diff --git a/Engine/Plugins/Animation/SkeletalMeshModelingTools/Source/SkeletalMeshModelingTools/Private/Interfaces/ISKMBackedDynaMeshComponentProvider.h b/Engine/Plugins/Animation/SkeletalMeshModelingTools/Source/SkeletalMeshModelingTools/Private/Interfaces/ISKMBackedDynaMeshComponentProvider.h
new file mode 100644
--- /dev/null
+++ b/Engine/Plugins/Animation/SkeletalMeshModelingTools/Source/SkeletalMeshModelingTools/Private/Interfaces/ISKMBackedDynaMeshComponentProvider.h
@@ -0,0 +1,26 @@
+﻿// Copyright Epic Games, Inc. All Rights Reserved.
+
+#pragma once
+
+#include "UObject/Interface.h"
+
+#include "ISKMBackedDynaMeshComponentProvider.generated.h"
+
+class USkeletalMeshBackedDynamicMeshComponent;
+
+UINTERFACE(MinimalAPI)
+class USkeletalMeshBackedDynamicMeshComponentProvider :
+	public UInterface
+{
+	GENERATED_BODY()
+};
+
+
+class ISkeletalMeshBackedDynamicMeshComponentProvider
+{
+	GENERATED_BODY()
+
+public:
+	virtual USkeletalMeshBackedDynamicMeshComponent* GetComponent(UObject* SourceObject) = 0;
+	virtual bool ShouldCommitToSkeletalMeshOnToolCommit() = 0;
+};
diff --git a/Engine/Plugins/Animation/SkeletalMeshModelingTools/Source/SkeletalMeshModelingTools/Private/SMorphTargetManager.cpp b/Engine/Plugins/Animation/SkeletalMeshModelingTools/Source/SkeletalMeshModelingTools/Private/SMorphTargetManager.cpp
new file mode 100644
--- /dev/null
+++ b/Engine/Plugins/Animation/SkeletalMeshModelingTools/Source/SkeletalMeshModelingTools/Private/SMorphTargetManager.cpp
@@ -0,0 +1,595 @@
+// Copyright Epic Games, Inc. All Rights Reserved.
+#include "SMorphTargetManager.h"
+
+#include "SkeletalMeshModelingToolsCommands.h"
+#include "SPositiveActionButton.h"
+#include "Widgets/Input/SSearchBox.h"
+#include "Widgets/Input/SSpinBox.h"
+#include "Widgets/Input/SCheckBox.h"
+#include "Framework/Commands/UICommandList.h"
+#include "Framework/MultiBox/MultiBoxBuilder.h"
+#include "Editor/EditorEngine.h"
+#include "Internationalization/Regex.h"
+#include "ScopedTransaction.h"
+#include "Framework/Commands/GenericCommands.h"
+#include "Widgets/Text/SInlineEditableTextBlock.h"
+
+extern UNREALED_API UEditorEngine* GEditor;
+
+
+#define LOCTEXT_NAMESPACE "SMorphTargetList"
+
+namespace MorphTargetManagerLocal
+{
+	struct FMorphTargetInfo
+	{
+		FName Name = NAME_None;
+		float Weight = 1.0f;
+		TSharedPtr<SInlineEditableTextBlock> EditableText;
+	};
+	static const FName ColumnId_MorphTargetNameLabel( "MorphTargetName" );
+	static const FName ColumnID_MorphTargetWeightLabel( "Weight" );
+}
+
+class SMorphTargetManagerListRow : public SMultiColumnTableRow< FMorphTargetInfoPtr >
+{
+public:
+
+	SLATE_BEGIN_ARGS( SMorphTargetManagerListRow ) {}
+
+	/** The item for this row **/
+	SLATE_ARGUMENT( FMorphTargetInfoPtr, Item )
+
+		/* The SMorphTargetViewer that we push the morph target weights into */
+		SLATE_ARGUMENT( SMorphTargetManager* , Manager)
+
+		/* Widget used to display the list of morph targets */
+		SLATE_ARGUMENT( SMorphTargetManagerListType* , ListView )
+
+	SLATE_END_ARGS()
+
+	void Construct( const FArguments& InArgs, const TSharedRef<STableViewBase>& OwnerTableView )
+	{
+		Item = InArgs._Item;
+		Manager = InArgs._Manager;
+		ListView = InArgs._ListView;
+		
+		SMultiColumnTableRow< FMorphTargetInfoPtr >::Construct( FSuperRowType::FArguments(), OwnerTableView );
+	}
+
+	/** Overridden from SMultiColumnTableRow.  Generates a widget for this column of the tree row. */
+	virtual TSharedRef<SWidget> GenerateWidgetForColumn( const FName& ColumnName ) override
+	{
+		using namespace MorphTargetManagerLocal;
+		
+		if ( ColumnName == ColumnId_MorphTargetNameLabel )
+		{
+			FText MorphNameText = FText::FromName(Item->Name);
+			return
+				SNew( SHorizontalBox )
+				+ SHorizontalBox::Slot()
+				.AutoWidth()
+				.Padding( 4.0f , 4.0f, 0.0f, 4.0f )
+				.VAlign( VAlign_Center )
+				[
+					SNew(SCheckBox)
+					.ToolTipText(LOCTEXT("MorphTargetEditToggleToolTip", "Activate Morph Target for editing"))
+					.Style(FAppStyle::Get(), "RadioButton")
+					.OnCheckStateChanged(this, &SMorphTargetManagerListRow::OnSetEditingMorphTarget)
+					.IsChecked(this, &SMorphTargetManagerListRow::IsEditingMorphTarget)
+				]
+				+ SHorizontalBox::Slot()
+				.FillWidth(1.0f)
+				.Padding( 0.0f, 4.0f )
+				.VAlign( VAlign_Center )
+				[
+					SAssignNew(Item->EditableText, SInlineEditableTextBlock)
+					.Text_Lambda([this]()
+					{
+						return FText::FromName(Item->Name);
+					})
+					.HighlightText(Manager,  &SMorphTargetManager::GetHighlightText, MorphNameText)
+					.OnTextCommitted(this, &SMorphTargetManagerListRow::OnRenameMorphTarget)
+					.IsSelected(this, &SMorphTargetManagerListRow::IsSelected)
+				]
+				+ SHorizontalBox::Slot()
+				.AutoWidth()
+				.Padding( 4.0f, 4.0f)
+				.VAlign( VAlign_Center )
+				[
+					SNew(STextBlock)
+					.Text( LOCTEXT("EditingMorphTargetMarker", "(Editing)") )
+					.Visibility_Lambda([this]()
+					{
+						return IsEditingMorphTarget() == ECheckBoxState::Checked ? EVisibility::Visible : EVisibility::Collapsed;
+					})
+				];
+		}
+		else if ( ColumnName == ColumnID_MorphTargetWeightLabel )
+		{
+			// Encase the SSpinbox in an SVertical box so we can apply padding. Setting ItemHeight on the containing SListView has no effect :-(
+			return
+				SNew( SHorizontalBox)
+				+ SHorizontalBox::Slot()
+				.FillWidth(1.0f)
+				.Padding( 0.0f, 4.0f )
+				.VAlign( VAlign_Center )
+				[
+					SNew( SSpinBox<float> )
+					.MinSliderValue(-1.f)
+					.MaxSliderValue(1.f)
+					.Value( this, &SMorphTargetManagerListRow::GetWeight )
+					.OnBeginSliderMovement(this, &SMorphTargetManagerListRow::OnBeginSlideMorphTargetWeight)
+					.OnEndSliderMovement(this, &SMorphTargetManagerListRow::OnEndSlideMorphTargetWeight)
+					.OnValueChanged( this, &SMorphTargetManagerListRow::OnMorphTargetWeightChanged )
+					.OnValueCommitted( this, &SMorphTargetManagerListRow::OnMorphTargetWeightValueCommitted )
+					.IsEnabled(this, &SMorphTargetManagerListRow::IsMorphTargetWeightSliderEnabled)
+				]
+				+ SHorizontalBox::Slot()
+				.AutoWidth()
+				.Padding( 0.0f, 4.0f )
+				.VAlign( VAlign_Center )
+				[
+					SNew(SCheckBox)
+					.ToolTipText(LOCTEXT("MorphTargetAutoFillToolTip", "When checked, animation system takes control of morph target weights"))
+					.OnCheckStateChanged(this, &SMorphTargetManagerListRow::OnMorphTargetAutoFillChecked)
+					.IsChecked(this, &SMorphTargetManagerListRow::IsMorphTargetAutoFillChangedChecked)
+				];
+		}
+		
+
+		return SNullWidget::NullWidget;
+	}
+
+	
+
+private:
+
+	/**
+	* Called when the user begins/ends dragging the slider on the SSpinBox
+	*/
+	void OnBeginSlideMorphTargetWeight()
+	{
+		GEditor->BeginTransaction(LOCTEXT("OverrideMorphTargetWeight", "Override Morph Target Weight"));
+	};
+	void OnEndSlideMorphTargetWeight(float Value)
+	{
+		GEditor->EndTransaction();
+	};
+	
+	/**
+	* Called when the user changes the value of the SSpinBox
+	*
+	* @param NewWeight - The new number the SSpinBox is set to
+	*
+	*/
+	void OnMorphTargetWeightChanged( float NewWeight )
+	{
+		Item->Weight = NewWeight;
+		Manager->SetMorphTargetWeight(Item->Name, NewWeight);
+	};
+	
+	/**
+	* Called when the user types the value and enters
+	*
+	* @param NewWeight - The new number the SSpinBox is set to
+	*
+	*/
+	void OnMorphTargetWeightValueCommitted( float NewWeight, ETextCommit::Type CommitType)
+	{
+		OnMorphTargetWeightChanged( NewWeight );
+	};
+	
+	/**
+	* Called to know if we enable or disable the weight sliders
+	*/
+	bool IsMorphTargetWeightSliderEnabled() const
+	{
+		return true;
+	};
+	
+	/** Auto fill check call back functions */
+	void OnMorphTargetAutoFillChecked(ECheckBoxState InState)
+	{
+		Manager->SetMorphTargetAutoFill(Item->Name, InState == ECheckBoxState::Checked, Item->Weight);
+	};
+	ECheckBoxState IsMorphTargetAutoFillChangedChecked() const
+	{
+		return Manager->GetMorphTargetAutoFill(Item->Name) ? ECheckBoxState::Checked : ECheckBoxState::Unchecked;
+	};
+
+	/** Auto fill check call back functions */
+	void OnSetEditingMorphTarget(ECheckBoxState InState)
+	{
+		FName EditingMorphTarget = InState == ECheckBoxState::Checked ? Item->Name : NAME_None;
+		Manager->SetEditingMorphTarget(EditingMorphTarget);
+	};
+	ECheckBoxState IsEditingMorphTarget() const
+	{
+		return Manager->IsEditingMorphTarget(Item->Name) ? ECheckBoxState::Checked : ECheckBoxState::Unchecked;
+	};
+	
+	/**
+	* Returns the weight of this morph target
+	*
+	* @return SearchText - The new number the SSpinBox is set to
+	*
+	*/
+	float GetWeight() const
+	{
+		return Manager->GetMorphTargetWeight(Item->Name);
+	};
+
+	void OnRenameMorphTarget(const FText& InNewName, ETextCommit::Type)
+	{
+		Item->Name = Manager->RenameMorphTarget( Item->Name, *InNewName.ToString());
+	}
+
+	/* The SMorphTargetViewer that we push the morph target weights into */
+	SMorphTargetManager* Manager = nullptr;
+
+	/** Widget used to display the list of morph targets */
+	SMorphTargetManagerListType* ListView = nullptr;
+
+	/** The name and weight of the morph target */
+	FMorphTargetInfoPtr Item;
+};
+
+
+
+void SMorphTargetManager::Construct(const FArguments& InArgs)
+{
+	using namespace MorphTargetManagerLocal;
+
+	Delegates = InArgs._Delegates;
+
+	BindCommands();
+	
+	ChildSlot
+	[
+		SNew( SVerticalBox )
+		+ SVerticalBox::Slot()
+		.AutoHeight()
+		.Padding(FMargin(0.f, 2.f))
+		[
+			SNew(SHorizontalBox)
+			+SHorizontalBox::Slot()
+			.AutoWidth()
+			.VAlign(VAlign_Center)
+			.Padding(FMargin(6.f, 0.0))
+			[
+				SNew(SPositiveActionButton)
+				.OnGetMenuContent( this, &SMorphTargetManager::CreateNewMenuWidget )
+				.Icon(FAppStyle::Get().GetBrush("Icons.Plus"))
+			]
+			// Filter entry
+			+SHorizontalBox::Slot()
+			.FillWidth(1.0f)
+			[
+				SAssignNew( NameFilterBox, SSearchBox )
+				.SelectAllTextWhenFocused( true )
+				.OnTextChanged( this, &SMorphTargetManager::OnFilterTextChanged )
+			]
+			+SHorizontalBox::Slot()
+			.AutoWidth()
+			.VAlign(VAlign_Center)
+			.Padding(FMargin(6.f, 0.0))
+			[
+				// A dummy widget to inform user the morph target viewer was replaced with morph target manager
+				// and how they can access the default viewer
+				SNew(SCheckBox)
+				.Style(FAppStyle::Get(), "ToggleButtonCheckbox")
+				.IsChecked_Lambda([](){return ECheckBoxState::Checked;})
+				[
+					SNew(STextBlock)
+					.ToolTipText(LOCTEXT("UsingMorphTargetManagerToolTip", "Deactivate Skeletal Mesh Editing Tools to access default Morph Target Viewer"))
+					.Text(LOCTEXT("UsingMorphTargetManagerNote", "Editing Tools Enabled"))
+				]
+			]
+		]
+
+		+ SVerticalBox::Slot()
+		.Padding( FMargin( 0.0f, 2.0f, 0.0f, 0.0f ) )
+		.FillHeight( 1.0f )		// This is required to make the scrollbar work, as content overflows Slate containers by default
+		[
+			SAssignNew( ListView , SMorphTargetManagerListType )
+			.ListItemsSource( &List)
+			.OnGenerateRow(this, &SMorphTargetManager::GenerateMorphTargetRow )
+			.OnContextMenuOpening( this, &SMorphTargetManager::OnGetContextMenuContent )
+			.HeaderRow
+			(
+				SNew( SHeaderRow )
+				+ SHeaderRow::Column( ColumnId_MorphTargetNameLabel )
+				.DefaultLabel( LOCTEXT( "MorphTargetNameLabel", "Morph Target Name" ) )
+
+				+ SHeaderRow::Column( ColumnID_MorphTargetWeightLabel )
+				.DefaultLabel( LOCTEXT( "MorphTargetWeightLabel", "Weight" ) )
+			)
+		]	
+	];
+
+	RefreshList();
+
+	if (GEditor)
+	{
+		GEditor->RegisterForUndo(this);
+	}
+}
+SMorphTargetManager::~SMorphTargetManager()
+{
+	if (GEditor)
+	{
+		GEditor->UnregisterForUndo(this);
+	}
+}
+
+void SMorphTargetManager::PostUndo(bool bSuccess)
+{
+	RefreshList();
+}
+
+void SMorphTargetManager::PostRedo(bool bSuccess)
+{
+	RefreshList();
+}
+
+void SMorphTargetManager::BindCommands()
+{
+	const FSkeletalMeshModelingToolsCommands& Commands = FSkeletalMeshModelingToolsCommands::Get();
+	const FGenericCommands& GenericCommands = FGenericCommands::Get();
+
+	CommandList = MakeShared<FUICommandList>(); 
+	CommandList->MapAction(Commands.NewMorphTarget, FSimpleDelegate::CreateSP(this, &SMorphTargetManager::AddMorphTarget) );
+	CommandList->MapAction(
+		GenericCommands.Rename,
+		FExecuteAction::CreateSP(this, &SMorphTargetManager::RenameSelectedMorphTarget), 
+		FCanExecuteAction::CreateSP(this, &SMorphTargetManager::CanRename) );
+	CommandList->MapAction(
+		GenericCommands.Delete,
+		FExecuteAction::CreateSP(this, &SMorphTargetManager::RemoveSelectedMorphTargets), 
+		FCanExecuteAction::CreateSP(this, &SMorphTargetManager::CanRemove) );
+	
+	CommandList->MapAction(
+		GenericCommands.Duplicate,
+		FExecuteAction::CreateSP(this, &SMorphTargetManager::DuplicateSelectedMorphTargets), 
+		FCanExecuteAction::CreateSP(this, &SMorphTargetManager::CanDuplicate));
+}
+
+void SMorphTargetManager::RefreshList()
+{
+	using namespace MorphTargetManagerLocal;
+	
+	List.Reset();
+
+	FText FilterText = GetFilterText();
+	
+	const TArray<FName>& MorphTargets = Delegates.OnGetMorphTargets.Execute();
+	
+	for (const FName& MorphTarget : MorphTargets)
+	{
+		bool bAddToList = true;
+		if (!FilterText.IsEmpty())
+		{
+			FRegexPattern Pattern(FilterText.ToString());
+			FRegexMatcher Matcher(Pattern, MorphTarget.ToString());
+
+			bAddToList = Matcher.FindNext();
+		}
+		
+		if (bAddToList)
+		{
+			TSharedPtr<FMorphTargetInfo> Info = MakeShared<FMorphTargetInfo>();
+			Info->Name = MorphTarget;
+			Info->Weight = Delegates.OnGetMorphTargetWeight.Execute(MorphTarget);
+			List.Add(Info);	
+		}
+	}
+
+	ListView->RequestListRefresh();
+}
+
+void SMorphTargetManager::SelectMorphTargets(const TArray<FName>& MorphTargets)
+{
+	ListView->ClearSelection();
+	
+	FMorphTargetInfoPtr First = nullptr; 
+	for (const FName& Name : MorphTargets)
+	{
+		FMorphTargetInfoPtr* Info = List.FindByPredicate([&](const FMorphTargetInfoPtr& InInfo)
+		{
+			return InInfo->Name == Name;
+		});
+
+		if (Info)
+		{
+			ListView->SetItemSelection(*Info, true);
+			if (!First)
+			{
+				First = *Info;
+			}
+		}
+	}
+	
+	ListView->RequestNavigateToItem(First);
+}
+
+TSharedRef<SWidget> SMorphTargetManager::CreateNewMenuWidget()
+{
+	const FSkeletalMeshModelingToolsCommands& Commands = FSkeletalMeshModelingToolsCommands::Get();
+
+	static constexpr bool bCloseAfterSelection = true;
+	FMenuBuilder MenuBuilder(bCloseAfterSelection, CommandList);
+	
+	MenuBuilder.BeginSection("NewMorphTarget", LOCTEXT("AddNewMorphTargetOperations", "Morph Targets"));
+	MenuBuilder.AddMenuEntry(Commands.NewMorphTarget);
+	MenuBuilder.EndSection();
+		
+	return MenuBuilder.MakeWidget();
+}
+
+void SMorphTargetManager::OnFilterTextChanged(const FText& Text)
+{
+	RefreshList();
+}
+
+FText SMorphTargetManager::GetFilterText() const
+{
+	return NameFilterBox->GetText(); 
+}
+
+FText SMorphTargetManager::GetHighlightText(FText InName) const
+{
+	FText FilterText = GetFilterText();
+	if (!FilterText.IsEmpty())
+	{
+		FRegexPattern Pattern(FilterText.ToString());
+		FString Name = InName.ToString();
+		FRegexMatcher Matcher(Pattern, Name);
+		if (Matcher.FindNext())
+		{
+			const int Begin = Matcher.GetMatchBeginning();	
+			const int End = Matcher.GetMatchEnding();
+			const int MatchSize = End - Begin;
+			FString Highlight = Name.Mid(Begin, MatchSize);
+			return FText::FromString(Highlight);
+		}
+	}
+
+	return {};
+}
+
+TSharedPtr<SWidget> SMorphTargetManager::OnGetContextMenuContent()
+{
+	const bool bShouldCloseWindowAfterMenuSelection = true;
+	FMenuBuilder MenuBuilder( bShouldCloseWindowAfterMenuSelection, CommandList);
+
+	MenuBuilder.BeginSection("MorphTargetManagerAction", LOCTEXT( "MorphsAction", "Selected Item Actions" ) );
+
+	MenuBuilder.AddMenuEntry(FGenericCommands::Get().Rename);
+	MenuBuilder.AddMenuEntry(FGenericCommands::Get().Delete);
+	MenuBuilder.AddMenuEntry(FGenericCommands::Get().Duplicate);
+
+	MenuBuilder.EndSection();
+	return MenuBuilder.MakeWidget();
+}
+
+void SMorphTargetManager::SetMorphTargetWeight(FName MorphTarget, float Weight)
+{
+	FScopedTransaction Transaction(LOCTEXT("OverrideMorphTargetWeight", "Override Morph Target Weight"));
+	Delegates.OnSetMorphTargetWeight.Execute(MorphTarget, Weight);
+}
+
+float SMorphTargetManager::GetMorphTargetWeight(FName MorphTarget)
+{
+	return Delegates.OnGetMorphTargetWeight.Execute(MorphTarget);
+}
+
+void SMorphTargetManager::SetMorphTargetAutoFill(FName MorphTarget, bool bAutoFill, float PreviousOverrideWeight)
+{
+	FText Title = bAutoFill ? LOCTEXT("ClearOverride", "Clear Morph Target Override") : LOCTEXT("OverrideMorphTargetWeight", "Override Morph Target Weight");
+	FScopedTransaction Transaction(Title);
+	Delegates.OnSetMorphTargetAutoFill.Execute(MorphTarget, bAutoFill, PreviousOverrideWeight);
+}
+
+bool SMorphTargetManager::GetMorphTargetAutoFill(FName MorphTarget)
+{
+	return Delegates.OnGetMorphTargetAutoFill.Execute(MorphTarget);
+}
+
+void SMorphTargetManager::SetEditingMorphTarget(FName MorphTarget)
+{
+	FText Title = MorphTarget!=NAME_None ? LOCTEXT("EnabledMorphTargetEditing", "Enable Morph Target for editing") : LOCTEXT("DisabledMorphTargetEditing", "Disabled Morph Target For Editing");
+	FScopedTransaction Transaction(Title);
+	Delegates.OnSetEditingMorphTarget.Execute(MorphTarget);
+}
+
+bool SMorphTargetManager::IsEditingMorphTarget(FName MorphTarget)
+{
+	return Delegates.OnGetEditingMorphTarget.Execute() == MorphTarget;
+}
+
+void SMorphTargetManager::AddMorphTarget()
+{
+	FScopedTransaction Transaction(LOCTEXT("AddMorphTarget", "Add Morph Target"));
+	FName NewMorphTarget = TEXT("NewMorphTarget");
+	NewMorphTarget = Delegates.OnAddNewMorphTarget.Execute(NewMorphTarget);
+
+	RefreshList();
+	SelectMorphTargets({NewMorphTarget});
+}
+
+FName SMorphTargetManager::RenameMorphTarget(FName InOldName, FName InNewName)
+{
+	FScopedTransaction Transaction(LOCTEXT("RenameMorphTarget", "Rename Morph Target"));
+	return Delegates.OnRenameMorphTarget.Execute(InOldName, InNewName);
+}
+
+void SMorphTargetManager::RenameSelectedMorphTarget()
+{
+	TArray<FMorphTargetInfoPtr> Items = ListView->GetSelectedItems();
+	Items[0]->EditableText->EnterEditingMode();
+}
+
+bool SMorphTargetManager::CanRename()
+{
+	return ListView->GetSelectedItems().Num() == 1;
+}
+
+void SMorphTargetManager::RemoveSelectedMorphTargets()
+{
+	TArray<FMorphTargetInfoPtr> Items = ListView->GetSelectedItems();
+
+	TArray<FName> Names;
+	for (const FMorphTargetInfoPtr& Item : Items)
+	{
+		Names.Add(Item->Name);
+	}
+
+	FText Title = Names.Num() > 1 ? LOCTEXT("RemoveMorph", "Remove Morph Targets") : LOCTEXT("RemoveMorphs", "Remove Morph Targets");
+	FScopedTransaction Transaction(Title);
+	Delegates.OnRemoveMorphTargets.Execute(Names);
+	RefreshList();
+}
+
+bool SMorphTargetManager::CanRemove()
+{
+	return ListView->GetSelectedItems().Num() > 0;
+}
+
+void SMorphTargetManager::DuplicateSelectedMorphTargets()
+{
+	TArray<FMorphTargetInfoPtr> Items = ListView->GetSelectedItems();
+
+	TArray<FName> Names;
+	for (const FMorphTargetInfoPtr& Item : Items)
+	{
+		Names.Add(Item->Name);
+	}
+
+	FText Title = Names.Num() > 1 ? LOCTEXT("RemoveMorph", "Remove Morph Targets") : LOCTEXT("RemoveMorphs", "Remove Morph Targets");
+	FScopedTransaction Transaction(Title);
+	TArray<FName> Duplicated = Delegates.OnDuplicateMorphTargets.Execute(Names);
+	RefreshList();
+	SelectMorphTargets(Duplicated);
+}
+
+bool SMorphTargetManager::CanDuplicate()
+{
+	return ListView->GetSelectedItems().Num() > 0;
+}
+
+
+TSharedRef<class ITableRow> SMorphTargetManager::GenerateMorphTargetRow(
+	FMorphTargetInfoPtr MorphTargetItem,
+	const TSharedRef<STableViewBase>& TableViewBase)
+{
+	TSharedPtr<SMorphTargetManagerListRow> Row =
+		SNew(SMorphTargetManagerListRow, TableViewBase)
+		.Item( MorphTargetItem )
+		.Manager(this)
+		.ListView(ListView.Get());
+	
+	return Row.ToSharedRef();
+}
+
+#undef LOCTEXT_NAMESPACE
diff --git a/Engine/Plugins/Animation/SkeletalMeshModelingTools/Source/SkeletalMeshModelingTools/Private/SMorphTargetManager.h b/Engine/Plugins/Animation/SkeletalMeshModelingTools/Source/SkeletalMeshModelingTools/Private/SMorphTargetManager.h
new file mode 100644
--- /dev/null
+++ b/Engine/Plugins/Animation/SkeletalMeshModelingTools/Source/SkeletalMeshModelingTools/Private/SMorphTargetManager.h
@@ -0,0 +1,116 @@
+// Copyright Epic Games, Inc. All Rights Reserved.
+#pragma once
+
+#include "SkeletalMeshNotifier.h"
+#include "Widgets/SCompoundWidget.h"
+#include "Widgets/Views/SListView.h"
+#include "Widgets/Views/STableViewBase.h"
+#include "EditorUndoClient.h"
+
+namespace MorphTargetManagerLocal
+{
+	struct FMorphTargetInfo;
+}
+typedef TSharedPtr<MorphTargetManagerLocal::FMorphTargetInfo> FMorphTargetInfoPtr;
+typedef SListView< FMorphTargetInfoPtr > SMorphTargetManagerListType;
+
+
+class SSearchBox;
+class UMorphTargetModifier;
+class SMorphTargetManager;
+class FUICommandList;
+
+
+
+class SMorphTargetManager : public SCompoundWidget, public FEditorUndoClient
+{
+public:
+	DECLARE_DELEGATE_RetVal(TArray<FName>, FOnGetMorphTargets);
+	DECLARE_DELEGATE_RetVal_OneParam(float, FOnGetMorphTargetWeight, FName);
+	DECLARE_DELEGATE_TwoParams(FOnSetMorphTargetWeight, FName, float); 
+	DECLARE_DELEGATE_RetVal_OneParam(bool, FOnGetMorphTargetAutoFill, FName);
+	DECLARE_DELEGATE_ThreeParams(FOnSetMorphTargetAutoFill, FName, bool, float);
+	DECLARE_DELEGATE_RetVal(FName, FOnGetEditingMorphTarget);
+	DECLARE_DELEGATE_OneParam(FOnSetEditingMorphTarget, FName);
+	DECLARE_DELEGATE_RetVal_OneParam(FName, FOnAddMorphTarget, FName);
+	DECLARE_DELEGATE_RetVal_TwoParams(FName, FOnRenameMorphTarget, FName, FName);
+	DECLARE_DELEGATE_OneParam(FOnRemoveMorphTargets, const TArray<FName>&);
+	DECLARE_DELEGATE_RetVal_OneParam(TArray<FName>, FOnDuplicateMorphTargets, const TArray<FName>&);
+
+	struct FMorphTargetManagerDelegates
+	{
+		FOnGetMorphTargets OnGetMorphTargets;
+
+		FOnGetMorphTargetWeight OnGetMorphTargetWeight;
+		FOnSetMorphTargetWeight OnSetMorphTargetWeight;
+		
+		FOnGetMorphTargetAutoFill OnGetMorphTargetAutoFill;
+		FOnSetMorphTargetAutoFill OnSetMorphTargetAutoFill;
+
+		FOnGetEditingMorphTarget OnGetEditingMorphTarget;
+		FOnSetEditingMorphTarget OnSetEditingMorphTarget;
+		
+		FOnAddMorphTarget OnAddNewMorphTarget;
+		FOnRenameMorphTarget OnRenameMorphTarget;
+		FOnRemoveMorphTargets OnRemoveMorphTargets;
+		FOnDuplicateMorphTargets OnDuplicateMorphTargets;
+	};	
+	
+	SLATE_BEGIN_ARGS( SMorphTargetManager ) {}
+		SLATE_ARGUMENT(FMorphTargetManagerDelegates, Delegates)
+	SLATE_END_ARGS()
+
+	void Construct(const FArguments& InArgs);
+	
+	/**
+	* Destructor - resets the morph targets
+	*
+	*/
+	virtual ~SMorphTargetManager();
+
+	virtual void PostUndo(bool bSuccess) override;
+	virtual void PostRedo(bool bSuccess) override;
+	
+	void BindCommands();
+	void RefreshList();
+	void SelectMorphTargets(const TArray<FName>& MorphTargets);
+	
+	TSharedRef<SWidget> CreateNewMenuWidget();
+	void OnFilterTextChanged(const FText& Text);
+	FText GetFilterText() const;
+	FText GetHighlightText(FText InName) const;
+	
+	TSharedRef<class ITableRow> GenerateMorphTargetRow(FMorphTargetInfoPtr MorphTargetItem, const TSharedRef<STableViewBase>& TableViewBase);
+	TSharedPtr<SWidget> OnGetContextMenuContent();
+
+	void SetMorphTargetWeight(FName MorphTarget, float Weight);
+	float GetMorphTargetWeight(FName MorphTarget);
+	void SetMorphTargetAutoFill(FName MorphTarget, bool bAutoFill, float PreviousOverrideWeight);
+	bool GetMorphTargetAutoFill(FName MorphTarget);
+	void SetEditingMorphTarget(FName MorphTarget);
+	bool IsEditingMorphTarget(FName MorphTarget);
+
+	void AddMorphTarget();
+	FName RenameMorphTarget(FName InNewName, FName InOldName);
+protected:
+	void RenameSelectedMorphTarget();
+	bool CanRename();
+	void RemoveSelectedMorphTargets();
+	bool CanRemove();
+	void DuplicateSelectedMorphTargets();
+	bool CanDuplicate();
+	
+	TSharedPtr<SSearchBox> NameFilterBox;
+	TSharedPtr<SMorphTargetManagerListType> ListView;
+
+	TArray<FMorphTargetInfoPtr> List;
+	TArray<FMorphTargetInfoPtr> FullList;
+
+	TWeakObjectPtr<UMorphTargetModifier> Modifier;
+
+	FMorphTargetManagerDelegates Delegates;
+
+	// Commands
+	TSharedPtr<FUICommandList> CommandList;
+};
+
diff --git a/Engine/Plugins/Animation/SkeletalMeshModelingTools/Source/SkeletalMeshModelingTools/Private/SReferenceSkeletonTree.cpp b/Engine/Plugins/Animation/SkeletalMeshModelingTools/Source/SkeletalMeshModelingTools/Private/SReferenceSkeletonTree.cpp
--- a/Engine/Plugins/Animation/SkeletalMeshModelingTools/Source/SkeletalMeshModelingTools/Private/SReferenceSkeletonTree.cpp
+++ b/Engine/Plugins/Animation/SkeletalMeshModelingTools/Source/SkeletalMeshModelingTools/Private/SReferenceSkeletonTree.cpp
@@ -104,6 +104,7 @@ void SBoneItem::Construct(const FArguments& InArgs)
 	    SAssignNew(InlineWidget, SInlineEditableTextBlock)
 		.Text(this, &SBoneItem::GetName)
 		// .Font(TextFont)
+	    .IsReadOnly(InArgs._IsReadOnly)
 		.ColorAndOpacity_Static(&SBoneItem::GetTextColor, InArgs._IsSelected)
 		.MaximumLength(NAME_SIZE-1)
 		.OnVerifyTextChanged(this, &SBoneItem::OnVerifyNameChanged)
@@ -216,6 +217,7 @@ void SReferenceSkeletonRow::Construct(const FArguments& InArgs, const TSharedRef
 	WeakTreeElement = InArgs._TreeElement;
 	WeakModifier = InArgs._WeakModifier;
 	Delegates = InArgs._Delegates;
+	bIsReadOnly = InArgs._IsReadOnly;
 
 	const FSuperRowType::FArguments Args = FSuperRowType::FArguments()
 		.OnDragDetected(InArgs._Delegates.OnDragDetected)
@@ -244,6 +246,7 @@ TSharedRef<SWidget> SReferenceSkeletonRow::GenerateWidgetForColumn(const FName&
 			.TreeElement(WeakTreeElement)
 			.Delegates(Delegates)
 			.IsSelected(this, &SReferenceSkeletonRow::IsSelected)
+			.IsReadOnly(bIsReadOnly)
 	];
 }
 
@@ -289,6 +292,7 @@ SReferenceSkeletonTree::~SReferenceSkeletonTree()
 void SReferenceSkeletonTree::Construct(const FArguments& InArgs)
 {
 	Modifier = InArgs._Modifier;
+	bIsReadOnly = InArgs._bIsReadOnly;
 	TextFilter = MakeShareable(new FTextFilterExpressionEvaluator(ETextFilterExpressionEvaluatorMode::BasicString));
 	CommandList = MakeShared<FUICommandList>();
 	BindCommands();
@@ -319,6 +323,7 @@ void SReferenceSkeletonTree::Construct(const FArguments& InArgs)
 					SNew(SPositiveActionButton)
 					.OnGetMenuContent(this, &SReferenceSkeletonTree::CreateAddNewMenu)
 					.Icon(FAppStyle::Get().GetBrush("Icons.Plus"))
+					.Visibility(this, &SReferenceSkeletonTree::GetAddNewButtonVisibility)
 				]
 
 				+SHorizontalBox::Slot()
@@ -348,7 +353,8 @@ void SReferenceSkeletonTree::Construct(const FArguments& InArgs)
 						return SNew(SReferenceSkeletonRow, OwnerTable)
 							.WeakModifier(Modifier)
 							.TreeElement(InItem)
-							.Delegates(Delegates);
+							.Delegates(Delegates)
+							.IsReadOnly(bIsReadOnly);
 					})
 					.OnGetChildren(this, &SReferenceSkeletonTree::HandleGetChildrenForTree)
 					.OnSelectionChanged(this, &SReferenceSkeletonTree::OnSelectionChanged)
@@ -505,6 +511,11 @@ void SReferenceSkeletonTree::HandleNewBone()
 
 bool SReferenceSkeletonTree::CanAddNewBone() const
 {
+	if (bIsReadOnly)
+	{
+		return false;
+	}
+	
 	return true;
 }
 
@@ -538,6 +549,11 @@ void SReferenceSkeletonTree::HandleDeleteBone()
 
 bool SReferenceSkeletonTree::CanDeleteBone() const
 {
+	if (bIsReadOnly)
+	{
+		return false;
+	}
+	
 	return HasSelectedItems();
 }
 
@@ -573,6 +589,11 @@ void SReferenceSkeletonTree::HandleUnParentBone()
 
 bool SReferenceSkeletonTree::CanUnParentBone() const
 {
+	if (bIsReadOnly)
+	{
+		return false;
+	}
+	
 	return HasSelectedItems();
 }
 
@@ -596,6 +617,11 @@ void SReferenceSkeletonTree::HandleCopyBones() const
 
 bool SReferenceSkeletonTree::CanCopyBones() const
 {
+	if (bIsReadOnly)
+	{
+		return false;
+	}
+
 	return Modifier.IsValid() && HasSelectedItems();
 }
 
@@ -631,6 +657,11 @@ void SReferenceSkeletonTree::HandlePasteBones()
 
 bool SReferenceSkeletonTree::CanPasteBones() const
 {
+	if (bIsReadOnly)
+	{
+		return false;
+	}
+	
 	return Modifier.IsValid() && SkeletonClipboard::IsClipboardValid();
 }
 
@@ -698,13 +729,14 @@ void SReferenceSkeletonTree::SelectItemFromNames(const TArray<FName>& InBoneName
 	}
 }
 
-ISkeletalMeshNotifier& SReferenceSkeletonTree::GetNotifier()
+TSharedPtr<ISkeletalMeshNotifier> SReferenceSkeletonTree::GetNotifier()
 {
 	if (!Notifier)
 	{
-		Notifier.Reset(new FReferenceSkeletonWidgetNotifier(SharedThis(this)));
+		Notifier = MakeShared<FReferenceSkeletonWidgetNotifier>(SharedThis(this));
 	}
-	return *Notifier;
+	
+	return Notifier;
 }
 
 void SReferenceSkeletonTree::HandleRenameBone(const ReferenceSkeletonTreeLocals::EBoneRenameType InRenameType) const
@@ -719,6 +751,11 @@ void SReferenceSkeletonTree::HandleRenameBone(const ReferenceSkeletonTreeLocals:
 
 bool SReferenceSkeletonTree::CanRenameBone() const
 {
+	if (bIsReadOnly)
+	{
+		return false;
+	}
+	
 	return HasSelectedItems();
 }
 
@@ -1044,8 +1081,18 @@ TSharedRef<SWidget> SReferenceSkeletonTree::CreateAddNewMenu()
 	return MenuBuilder.MakeWidget();
 }
 
+EVisibility SReferenceSkeletonTree::GetAddNewButtonVisibility() const
+{
+	return bIsReadOnly ? EVisibility::Collapsed : EVisibility::Visible;
+}
+
 TSharedPtr<SWidget> SReferenceSkeletonTree::CreateContextMenu()
 {
+	if (bIsReadOnly)
+	{
+		return nullptr;
+	}
+	
 	static constexpr bool bCloseAfterSelection = true;
 	FMenuBuilder MenuBuilder(bCloseAfterSelection, CommandList);
 
@@ -1164,6 +1211,11 @@ FReply SReferenceSkeletonTree::OnDragDetected(const FGeometry& MyGeometry, const
 
 TOptional<EItemDropZone> SReferenceSkeletonTree::OnCanAcceptDrop(const FDragDropEvent& DragDropEvent, EItemDropZone DropZone, TSharedPtr<FBoneElement> TargetItem)
 {
+	if (bIsReadOnly)
+	{
+		return TOptional<EItemDropZone>();
+		
+	}
 	TOptional<EItemDropZone> ReturnedDropZone;
 
 	const TSharedPtr<FBoneItemDragDropOp> DragDropOp = DragDropEvent.GetOperationAs<FBoneItemDragDropOp>();
diff --git a/Engine/Plugins/Animation/SkeletalMeshModelingTools/Source/SkeletalMeshModelingTools/Private/SReferenceSkeletonTree.h b/Engine/Plugins/Animation/SkeletalMeshModelingTools/Source/SkeletalMeshModelingTools/Private/SReferenceSkeletonTree.h
--- a/Engine/Plugins/Animation/SkeletalMeshModelingTools/Source/SkeletalMeshModelingTools/Private/SReferenceSkeletonTree.h
+++ b/Engine/Plugins/Animation/SkeletalMeshModelingTools/Source/SkeletalMeshModelingTools/Private/SReferenceSkeletonTree.h
@@ -97,6 +97,7 @@ class SBoneItem : public SCompoundWidget
 		SLATE_ARGUMENT(TWeakObjectPtr<USkeletonModifier>, WeakModifier)
 		SLATE_ARGUMENT(FRefSkeletonTreeDelegates, Delegates)
 		SLATE_EVENT( FIsSelected, IsSelected )
+		SLATE_ATTRIBUTE(bool, IsReadOnly)
 	SLATE_END_ARGS()
 
 	virtual ~SBoneItem() override;
@@ -146,6 +147,7 @@ public:
 		SLATE_ARGUMENT(TWeakObjectPtr<USkeletonModifier>, WeakModifier)
 		SLATE_ARGUMENT(TSharedPtr<FBoneElement>, TreeElement)
 		SLATE_ARGUMENT(FRefSkeletonTreeDelegates, Delegates)
+		SLATE_ARGUMENT(bool, IsReadOnly)
 	SLATE_END_ARGS()
 
 	// Slate construction function
@@ -159,6 +161,7 @@ private:
 	TWeakPtr<FBoneElement> WeakTreeElement;
 	TWeakObjectPtr<USkeletonModifier> WeakModifier;
 	FRefSkeletonTreeDelegates Delegates;
+	bool bIsReadOnly = false;
 };
 
 typedef STreeView< TSharedPtr<FBoneElement> > SRefSkeletonTreeView;
@@ -181,6 +184,7 @@ public:
 	SLATE_BEGIN_ARGS(SReferenceSkeletonTree) {}
 		SLATE_ARGUMENT(TWeakObjectPtr<USkeletonModifier>, Modifier)
 		SLATE_ARGUMENT(FRefSkeletonTreeDelegates, Delegates)
+		SLATE_ARGUMENT(bool, bIsReadOnly)
 	SLATE_END_ARGS()
 
 	SReferenceSkeletonTree();
@@ -198,7 +202,7 @@ public:
 	bool HasSelectedItems() const;
 	void SelectItemFromNames(const TArray<FName>& InBoneNames, bool bFrameSelection = false);
 
-	ISkeletalMeshNotifier& GetNotifier();
+	TSharedPtr<ISkeletalMeshNotifier> GetNotifier();
 
 private:
 	// SWidget interface
@@ -238,6 +242,7 @@ private:
 	void OnSelectionChanged(TSharedPtr<FBoneElement> InItem, ESelectInfo::Type InSelectInfo);
 	
 	TSharedRef< SWidget > CreateAddNewMenu();
+	EVisibility GetAddNewButtonVisibility() const;
 	TSharedPtr< SWidget > CreateContextMenu();
 
 	void OnItemDoubleClicked(TSharedPtr<FBoneElement> InItem);
@@ -255,7 +260,7 @@ private:
 	// The modifier that actually does modifications on the reference skeleton
 	TWeakObjectPtr<USkeletonModifier> Modifier;
 	
-	TUniquePtr<FReferenceSkeletonWidgetNotifier> Notifier;
+	TSharedPtr<FReferenceSkeletonWidgetNotifier> Notifier;
 
 	// undo
 	void BeginChange();
@@ -274,6 +279,10 @@ private:
 	// filter
 	TSharedPtr<FTextFilterExpressionEvaluator> TextFilter;
 
+
+	bool bIsReadOnly = false;
+	
+
 	friend class FReferenceSkeletonWidgetNotifier;
 };
 
diff --git a/Engine/Plugins/Animation/SkeletalMeshModelingTools/Source/SkeletalMeshModelingTools/Private/SkeletalMeshEditingCache.cpp b/Engine/Plugins/Animation/SkeletalMeshModelingTools/Source/SkeletalMeshModelingTools/Private/SkeletalMeshEditingCache.cpp
new file mode 100644
--- /dev/null
+++ b/Engine/Plugins/Animation/SkeletalMeshModelingTools/Source/SkeletalMeshModelingTools/Private/SkeletalMeshEditingCache.cpp
@@ -0,0 +1,804 @@
+﻿// Copyright Epic Games, Inc. All Rights Reserved.
+
+#include "SkeletalMeshEditingCache.h"
+
+#include "CanvasItem.h"
+#include "CanvasTypes.h"
+#include "PreviewMesh.h"
+#include "SkeletalDebugRendering.h"
+#include "Components/SKMBackedDynaMeshComponent.h"
+#include "Preferences/PersonaOptions.h"
+#include "SkeletalMesh/RefSkeletonPoser.h"
+#include "SkeletalMesh/SkeletalMeshEditingInterface.h"
+#include "Animation/AnimInstance.h"
+#include "DynamicMesh/MeshNormals.h"
+#include "UnrealClient.h"
+#include "EditorViewportClient.h"
+#include "SceneView.h"
+#include "Engine/Engine.h"
+
+
+#include UE_INLINE_GENERATED_CPP_BY_NAME(SkeletalMeshEditingCache)
+
+#define LOCTEXT_NAMESPACE "SkeletalMeshEditingCache"
+
+namespace SkeletalMeshEditingCacheLocal
+{
+	static FIntPoint GetDPIUnscaledSize(FViewport* Viewport, FViewportClient* Client)
+	{
+		const FIntPoint Size = Viewport->GetSizeXY();
+		const float DPIScale = Client->GetDPIScale();
+		// (FIntPoint / float) implicitly casts the float to an int if you try to divide it directly
+		return FIntPoint(static_cast<int32>(Size.X / DPIScale), static_cast<int32>(Size.Y / DPIScale));
+	}
+}
+
+FSkeletalMeshEditingCacheNotifier::FSkeletalMeshEditingCacheNotifier(USkeletalMeshEditingCache* InEditingCahe)
+	:EditingCache(InEditingCahe)
+{
+}
+
+void FSkeletalMeshEditingCacheNotifier::HandleNotification(const TArray<FName>& BoneNames, const ESkeletalMeshNotifyType InNotifyType)
+{
+	if (Notifying())
+	{
+		return;
+	}
+
+	EditingCache->HandleNotification(BoneNames, InNotifyType);
+
+	Notify(BoneNames, InNotifyType);
+}
+
+void USkeletalMeshEditingCache::Spawn(UWorld* World, USkeletalMeshComponent* InSkeletalMeshComponent, EMeshLODIdentifier InLOD, const FDelegates& InDelegates)
+{
+	Delegates = InDelegates;
+	SkeletalMeshComponent = InSkeletalMeshComponent;
+	
+	HostActor = World->SpawnActor(AActor::StaticClass());
+	EditingMeshComponent = NewObject<USkeletalMeshBackedDynamicMeshComponent>(HostActor);
+
+	USkeletalMesh* SkeletalMesh = SkeletalMeshComponent->GetSkeletalMeshAsset();
+	EditingMeshComponent->Init(SkeletalMesh, InLOD);
+
+	TArray<UMaterialInterface*> Materials = SkeletalMeshComponent->GetMaterials();
+	for (int k = 0; k < Materials.Num(); ++k)
+	{
+		EditingMeshComponent->SetMaterial(k, Materials[k]);
+	}
+	
+	EditingMeshComponent->GetOnRequestingVisibilityChange().AddUObject(this, &USkeletalMeshEditingCache::HandleVisibilityChangeRequest);
+
+	HostActor->AddInstanceComponent(EditingMeshComponent);
+	HostActor->SetRootComponent(EditingMeshComponent);
+	HostActor->RegisterAllComponents();
+
+	FTransform ActorTransform = SkeletalMeshComponent->GetComponentTransform();
+	
+	HostActor->SetActorTransform(ActorTransform);
+		
+	FTransform PreviewActorTransform = ActorTransform ;
+	
+	PreviewMesh = NewObject<UPreviewMesh>();
+	PreviewMesh->CreateInWorld(World ,PreviewActorTransform);
+
+	UDynamicMesh* EditingDynamicMesh = EditingMeshComponent->GetDynamicMesh();
+	EditingDynamicMesh->ProcessMesh([&](const FDynamicMesh3& ReadMesh)
+	{
+		FDynamicMesh3 PreviewMeshObject = ReadMesh;
+		PreviewMesh->ReplaceMesh(MoveTemp(PreviewMeshObject));
+	});
+
+	PreviewMesh->SetMaterials(Materials);
+
+	RefSkeletonPoser = NewObject<URefSkeletonPoser>();
+	RefSkeletonPoser->SetRefSkeleton(EditingMeshComponent->GetRefSkeleton());
+	
+	EditingMeshComponent->SetVisibility(false);
+	PreviewMesh->SetVisible(false);
+	bEnableDynamicMesh = false;
+	bEnableDynamicMeshSkeleton = false;
+
+	bCacheVisibility = true;
+	CacheSkeletonDrawMode = ESkeletonDrawMode::Default;
+	bCacheBoneManipulation = true;
+
+	SetSkeletalMeshSkeletonDrawMode(CacheSkeletonDrawMode);
+	
+	ChangeCountWatcher.Initialize(
+		[this]()
+			{
+				return GetEditingMeshComponent()->GetChangeCount();
+			},
+		[this](int32 ChangeCount)
+			{
+				bEnableDynamicMesh = ChangeCount > 0;
+				UDynamicMesh* DynamicMesh = EditingMeshComponent->GetDynamicMesh();
+				FDynamicMesh3 PreviewMeshObject;
+				DynamicMesh->ProcessMesh([&](const FDynamicMesh3& ReadMesh)
+					{
+						PreviewMeshObject.Copy(ReadMesh);
+					});
+
+				PreviewMesh->ReplaceMesh(MoveTemp(PreviewMeshObject));	
+				
+			},
+		GetEditingMeshComponent()->GetChangeCount());
+
+	SkeletonChangeCountWatcher.Initialize(
+		[this]()
+			{
+				return GetEditingMeshComponent()->GetSkeletonChangeTracker().GetChangeCount();
+			},
+		[this](int32 SkeletonChangeCount)
+			{
+				bEnableDynamicMeshSkeleton = SkeletonChangeCount > 0;
+				RefSkeletonPoser->SetRefSkeleton(GetEditingMeshComponent()->GetRefSkeleton());
+				UpdateSelectedBoneIndices();
+			},
+		GetEditingMeshComponent()->GetSkeletonChangeTracker().GetChangeCount());
+
+	MeshVisibilityWatcher.Initialize(
+		[this]()
+			{
+				FMeshVisibilityState State;
+				State.bEnableDynamicMesh = bEnableDynamicMesh;
+				State.bCacheVisibility = bCacheVisibility;
+				return State;
+			},
+		[this](FMeshVisibilityState State)
+			{
+				PreviewMesh->SetVisible(
+					State.bEnableDynamicMesh && State.bCacheVisibility);
+
+				if (SkeletalMeshComponent.IsValid())
+				{
+					SkeletalMeshComponent->SetVisibility(
+						!State.bEnableDynamicMesh && State.bCacheVisibility);
+				}
+			},
+		FMeshVisibilityState(bEnableDynamicMesh, bCacheVisibility)
+	);
+
+	SkeletonVisibilityWatcher.Initialize(
+		[this]()
+			{
+				FSkeletonVisibilityState State;
+				State.bEnableDynamicMeshSkeleton = bEnableDynamicMeshSkeleton;
+				State.CacheSkeletonDrawMode = CacheSkeletonDrawMode;
+				return State;
+			},
+		[this](FSkeletonVisibilityState State)
+			{
+				if (State.bEnableDynamicMeshSkeleton)
+				{
+					DynamicMeshSkeletonDrawMode = State.CacheSkeletonDrawMode;
+					SetSkeletalMeshSkeletonDrawMode(ESkeletonDrawMode::Hidden);
+				}
+				else
+				{
+					DynamicMeshSkeletonDrawMode = ESkeletonDrawMode::Hidden;
+					SetSkeletalMeshSkeletonDrawMode(State.CacheSkeletonDrawMode);
+				}
+			},
+		FSkeletonVisibilityState(bEnableDynamicMeshSkeleton, CacheSkeletonDrawMode)
+	);
+
+	BoneManipulationWatcher.Initialize(
+		[this]()
+			{
+				FBoneManipulationState State;
+				State.bEnableDynamicMeshSkeleton = bEnableDynamicMeshSkeleton;
+				State.bCacheBoneManipulation = bCacheBoneManipulation;
+				return State;
+			},
+		[this](FBoneManipulationState State)
+			{
+				Delegates.ToggleSkeletalMeshBoneManipulationDelegate.ExecuteIfBound(
+					!State.bEnableDynamicMeshSkeleton && bCacheBoneManipulation);
+			},
+		FBoneManipulationState(bEnableDynamicMeshSkeleton, bCacheBoneManipulation)
+	);
+
+	PreviewMeshVisibilityWatcher.Initialize(
+		[this]()
+			{
+				return PreviewMesh->IsVisible();
+			},
+		[this](bool bIsVisible)
+			{
+				if (bIsVisible)
+				{
+					DeformPreviewMesh(GetComponentSpaceBoneTransforms(), GetMorphTargetWeights());
+				}
+			},
+		PreviewMesh->IsVisible());
+
+	PoseChangeDetector.GetNotifier().AddUObject(this, &USkeletalMeshEditingCache::HandlePoseChangeDetectorEvent);
+}
+
+EMeshLODIdentifier USkeletalMeshEditingCache::GetLOD() const
+{
+	return EditingMeshComponent->GetLOD();
+}
+
+
+void USkeletalMeshEditingCache::Destroy()
+{
+	HostActor->Destroy();
+	PreviewMesh->Disconnect();
+	
+	if (SkeletalMeshComponent.IsValid())
+	{
+		SkeletalMeshComponent->SetVisibility(true);
+	}
+}
+
+void USkeletalMeshEditingCache::ApplyChanges()
+{
+
+
+	if (SkeletalMeshComponent.IsValid())
+	{
+		// See USkeletalMeshComponentToolTarget
+		
+		// Unregister the component while we update its skeletal mesh
+		FComponentReregisterContext ComponentReregisterContext(SkeletalMeshComponent.Get());
+
+		EditingMeshComponent->CommitToSkeletalMesh();
+
+		// this rebuilds physics, but it doesn't undo!
+		SkeletalMeshComponent->RecreatePhysicsState();
+	}
+	else
+	{
+		EditingMeshComponent->CommitToSkeletalMesh();
+	}
+}
+
+bool USkeletalMeshEditingCache::IsApplyingChanges() const
+{
+	return EditingMeshComponent->IsCommitting();
+}
+
+void USkeletalMeshEditingCache::DiscardChanges()
+{
+	EditingMeshComponent->DiscardChanges();
+}
+
+void USkeletalMeshEditingCache::Tick()
+{
+	ChangeCountWatcher.CheckAndUpdate();
+	SkeletonChangeCountWatcher.CheckAndUpdate();
+	MeshVisibilityWatcher.CheckAndUpdate();
+	SkeletonVisibilityWatcher.CheckAndUpdate();
+	BoneManipulationWatcher.CheckAndUpdate();
+	PreviewMeshVisibilityWatcher.CheckAndUpdate();
+	
+	if (bEnableDynamicMesh)
+	{
+		PoseChangeDetector.CheckPose( GetComponentSpaceBoneTransforms(),  GetMorphTargetWeights());
+	}
+}
+
+bool USkeletalMeshEditingCache::HandleClick(HHitProxy* HitProxy)
+{
+	if (bEnableDynamicMeshSkeleton)
+	{
+		TArray<FName> SelectedBones;	
+		if (const HBoneHitProxy* BoneHitProxy = HitProxyCast<HBoneHitProxy>(HitProxy))
+		{
+			SelectedBones.Emplace(BoneHitProxy->BoneName);
+		}
+
+		Notifier->HandleNotification(SelectedBones, ESkeletalMeshNotifyType::BonesSelected);
+	
+		return true;
+	}
+
+	return false;
+}
+
+void USkeletalMeshEditingCache::Render(FPrimitiveDrawInterface* PDI, TFunction<void(FSkelDebugDrawConfig&)> OverrideBoneDrawConfigFunc)
+{
+	if (DynamicMeshSkeletonDrawMode == ESkeletonDrawMode::Hidden)
+	{
+		return;
+	}
+
+	const FReferenceSkeleton& RefSkeleton = EditingMeshComponent->GetRefSkeleton();
+	const TArray<FTransform>& ComponentSpaceTransforms = GetComponentSpaceBoneTransforms();
+	const FTransform& PreviewActorTransform = PreviewMesh->GetActor()->GetActorTransform();
+	
+	const int32 NumBones = RefSkeleton.GetRawBoneNum();	
+
+	TArray<FTransform> WorldTransforms;
+	WorldTransforms.AddUninitialized(NumBones);
+
+	TArray<FLinearColor> BoneColors;
+	BoneColors.AddUninitialized(NumBones);
+
+	TArray<FBoneIndexType> RequiredBones;
+	RequiredBones.AddUninitialized(NumBones);
+	
+	for ( int32 BoneIndex=0; BoneIndex<NumBones; ++BoneIndex )
+	{
+		RequiredBones[BoneIndex] = BoneIndex;
+		WorldTransforms[BoneIndex] = ComponentSpaceTransforms[BoneIndex] * PreviewActorTransform ;
+		BoneColors[BoneIndex] = GetDefaultBoneColor(BoneIndex);
+	}
+
+	constexpr bool bForceDraw = false;
+	constexpr bool bAddHitProxy = true;
+		
+	const bool bUseMultiColors = GetDefault<UPersonaOptions>()->bShowBoneColors;
+	
+	FSkelDebugDrawConfig DrawConfig;
+	DrawConfig.BoneDrawMode = EBoneDrawMode::Type::All ;
+	DrawConfig.BoneDrawSize = 1.0f;
+	DrawConfig.bAddHitProxy = bAddHitProxy;
+	DrawConfig.bForceDraw = bForceDraw;
+	DrawConfig.bUseMultiColorAsDefaultColor = bUseMultiColors;
+	DrawConfig.DefaultBoneColor = GetMutableDefault<UPersonaOptions>()->DefaultBoneColor;
+	DrawConfig.AffectedBoneColor = GetMutableDefault<UPersonaOptions>()->AffectedBoneColor;
+	DrawConfig.SelectedBoneColor = GetMutableDefault<UPersonaOptions>()->SelectedBoneColor;
+	DrawConfig.ParentOfSelectedBoneColor = GetMutableDefault<UPersonaOptions>()->ParentOfSelectedBoneColor;
+
+	OverrideBoneDrawConfigFunc(DrawConfig);
+
+	TArray<TRefCountPtr<HHitProxy>> HitProxies; HitProxies.Reserve(NumBones);
+
+	for (int32 Index = 0; Index < NumBones; ++Index)
+	{
+		HitProxies.Add(new HBoneHitProxy(Index, RefSkeleton.GetBoneName(Index)));
+	}
+
+	SkeletalDebugRendering::DrawBones(
+		PDI,
+		PreviewActorTransform.GetLocation(),
+		RequiredBones,
+		RefSkeleton,
+		WorldTransforms,
+		SelectedBoneIndices,
+		BoneColors,
+		HitProxies,
+		DrawConfig
+	);
+}
+
+void USkeletalMeshEditingCache::DrawHUD(FEditorViewportClient* ViewportClient, FViewport* Viewport, const FSceneView* View, FCanvas* Canvas)
+{
+	if (bEnableDynamicMeshSkeleton)
+	{
+		// See FSkeletonSelectionEditMode::DrawHUD
+		const FReferenceSkeleton& RefSkeleton = GetEditingMeshComponent()->GetRefSkeleton();
+		const int32 BoneIndex = GetFirstSelectedBoneIndex();
+
+		// Draw name of selected bone
+		if (RefSkeleton.IsValidIndex(BoneIndex))
+		{
+			const FIntPoint ViewPortSize = SkeletalMeshEditingCacheLocal::GetDPIUnscaledSize(Viewport, ViewportClient);
+			const int32 HalfX = ViewPortSize.X / 2;
+			const int32 HalfY = ViewPortSize.Y / 2;
+
+			const FName BoneName = RefSkeleton.GetBoneName(BoneIndex);
+
+			const FMatrix BoneMatrix = (GetComponentSpaceBoneTransforms()[BoneIndex] * GetTransform()).ToMatrixNoScale();
+			const FPlane Proj = View->Project(BoneMatrix.GetOrigin());
+			if (Proj.W > 0.f)
+			{
+				const int32 XPos = HalfX + static_cast<int32>(HalfX * Proj.X);
+				const int32 YPos = HalfY + static_cast<int32>(HalfY * Proj.Y * -1);
+
+				FCanvasTextItem TextItem(FVector2D(XPos, YPos), FText::FromString(BoneName.ToString()), GEngine->GetSmallFont(), FLinearColor::White);
+				TextItem.EnableShadow(FLinearColor::Black);
+				Canvas->DrawItem(TextItem);
+			}
+		}
+	}
+}
+
+TSharedPtr<ISkeletalMeshNotifier> USkeletalMeshEditingCache::GetNotifier()
+{
+	if (!Notifier.IsValid())
+	{
+		Notifier = MakeShared<FSkeletalMeshEditingCacheNotifier>(this);
+	}
+
+	return Notifier;
+}
+
+void USkeletalMeshEditingCache::HandleNotification(const TArray<FName>& BoneNames, const ESkeletalMeshNotifyType InNotifyType)
+{
+	switch (InNotifyType)
+	{
+	case ESkeletalMeshNotifyType::BonesSelected:
+		SelectedBoneNames = BoneNames;
+		UpdateSelectedBoneIndices();
+		break;
+	}
+}
+
+
+USkeletalMeshBackedDynamicMeshComponent* USkeletalMeshEditingCache::GetEditingMeshComponent() const
+{
+	return EditingMeshComponent;
+}
+
+
+void USkeletalMeshEditingCache::HandleVisibilityChangeRequest(bool bVisible)
+{
+	bCacheVisibility = bVisible;
+}
+
+
+
+void USkeletalMeshEditingCache::HandlePoseChangeDetectorEvent(SkeletalMeshToolsHelper::FPoseChangeDetector::FPayload Payload)
+{
+	using namespace SkeletalMeshToolsHelper;
+	
+	if (Payload.CurrentState == FPoseChangeDetector::PoseJustChanged ||
+		Payload.CurrentState == FPoseChangeDetector::PoseChanged)
+	{
+		DeformPreviewMesh(Payload.ComponentSpaceTransforms, Payload.MorphTargetWeights);
+	}
+}
+
+const TArray<FTransform>& USkeletalMeshEditingCache::GetComponentSpaceBoneTransforms() const
+{
+	if (bEnableDynamicMeshSkeleton)
+	{
+		return RefSkeletonPoser->GetComponentSpaceTransforms();
+	}
+
+	return SkeletalMeshComponent->GetComponentSpaceTransforms();
+}
+
+FTransform USkeletalMeshEditingCache::GetTransform()
+{
+	return PreviewMesh->GetTransform();
+}
+
+void USkeletalMeshEditingCache::HideSkeleton()
+{
+	SavedSkeletonDrawMode = CacheSkeletonDrawMode;
+	CacheSkeletonDrawMode = ESkeletonDrawMode::Hidden;
+}
+
+void USkeletalMeshEditingCache::ShowSkeleton()
+{
+	CacheSkeletonDrawMode = SavedSkeletonDrawMode;
+}
+
+void USkeletalMeshEditingCache::ToggleBoneManipulation(bool bEnable)
+{
+	bCacheBoneManipulation = bEnable;
+}
+
+
+bool USkeletalMeshEditingCache::IsDynamicMeshSkeletonEnabled() const
+{
+	return bEnableDynamicMeshSkeleton;
+}
+
+bool USkeletalMeshEditingCache::IsDynamicMeshBoneManipulationEnabled() const
+{
+	return IsDynamicMeshSkeletonEnabled() && bCacheBoneManipulation;
+}
+
+int32 USkeletalMeshEditingCache::GetFirstSelectedBoneIndex() const
+{
+	return !SelectedBoneIndices.IsEmpty() ? SelectedBoneIndices[0] : INDEX_NONE;
+}
+
+TArray<FName> USkeletalMeshEditingCache::GetSelectedBones() const
+{
+	return SelectedBoneNames;
+}
+
+void USkeletalMeshEditingCache::ResetDynamicMeshBoneTransforms(bool bSelectedOnly)
+{
+	if (!bEnableDynamicMeshSkeleton)
+	{
+		return;
+	}
+	
+	RefSkeletonPoser->BeginPoseChange();
+	if (bSelectedOnly)
+	{
+		for (int32 BoneIndex : SelectedBoneIndices)
+		{
+			RefSkeletonPoser->ClearBoneAdditiveTransform(BoneIndex);
+		}
+	}
+	else
+	{
+		RefSkeletonPoser->ClearAllBoneAdditiveTransforms();
+	}
+	RefSkeletonPoser->EndPoseChange();
+}
+
+TArray<FName> USkeletalMeshEditingCache::GetMorphTargets() const
+{
+	return EditingMeshComponent->GetMorphTargetChangeTracker().GetCurrentMorphTargetNames();
+}
+
+TMap<FName, float> USkeletalMeshEditingCache::GetMorphTargetWeights() const
+{
+	TMap<FName, float> MorphTargetWeights;
+
+	for (const FName& Name : GetMorphTargets())
+	{
+		MorphTargetWeights.Emplace(Name, 0.0f);
+	}
+	
+	const TMap<FName, float>& SkeletalMeshWeights = GetSkeletalMeshComponentMorphTargetWeights();
+
+	for (const TPair<FName, float>& SkeletalMeshWeight : SkeletalMeshWeights)
+	{
+		FName UpdatedName = EditingMeshComponent->GetMorphTargetChangeTracker().GetCurrentMorphTargetName(SkeletalMeshWeight.Key);
+		if (UpdatedName != NAME_None)
+		{
+			MorphTargetWeights[UpdatedName] = SkeletalMeshWeight.Value;
+		}
+	}
+
+	MorphTargetWeights.Append(MorphTargetOverrides);
+
+	return MorphTargetWeights;	
+}
+
+float USkeletalMeshEditingCache::GetMorphTargetWeight(FName MorphTarget) const
+{
+	if (const float* Weight = MorphTargetOverrides.Find(MorphTarget))
+	{
+		return *Weight;
+	}
+
+	const TMap<FName, float>& SkeletalMeshWeights = GetSkeletalMeshComponentMorphTargetWeights();
+
+	FName OriginalName = EditingMeshComponent->GetMorphTargetChangeTracker().GetOriginalMorphTargetName(MorphTarget);
+	
+	if (const float* Weight = SkeletalMeshWeights.Find(OriginalName))
+	{
+		return *Weight;
+	}
+
+	return 0.0f;
+}
+
+
+void USkeletalMeshEditingCache::HandleSetMorphTargetWeight(FName MorphTarget, float Weight)
+{
+	if (MorphTarget == NAME_None)
+	{
+		return;
+	}
+
+	if (GUndo)
+	{
+		SetFlags(RF_Transactional);
+		Modify();
+	}
+	
+	MorphTargetOverrides.Emplace(MorphTarget, Weight);
+	
+	FName OriginalMorphTargetName = EditingMeshComponent->GetMorphTargetChangeTracker().GetOriginalMorphTargetName(MorphTarget);
+	if (OriginalMorphTargetName != NAME_None)
+	{
+		if (UDebugSkelMeshComponent* Mesh = GetDebugSkelMeshComponent())
+		{
+			if (GUndo)
+			{
+				Mesh->SetFlags(RF_Transactional);
+				Mesh->Modify();
+			}
+			constexpr bool bRemoveZeroWeight = false;
+			Mesh->SetMorphTarget(OriginalMorphTargetName, Weight, bRemoveZeroWeight);
+		}
+	}
+}
+
+bool USkeletalMeshEditingCache::GetMorphTargetAutoFill(FName Name)
+{
+	return !MorphTargetOverrides.Contains(Name);
+}
+
+void USkeletalMeshEditingCache::HandleSetMorphTargetAutoFill(FName Name, bool bAutoFill, float PreviousOverrideWeight)
+{
+	if (GUndo)
+	{
+		SetFlags(RF_Transactional);
+		Modify();
+	}
+	
+	if (bAutoFill)
+	{
+		MorphTargetOverrides.Remove(Name);
+	}
+	else
+	{
+		MorphTargetOverrides.Emplace(Name, PreviousOverrideWeight);
+	}
+
+	if (UDebugSkelMeshComponent* Mesh = GetDebugSkelMeshComponent())
+	{
+		FName OriginalMorphTargetName = EditingMeshComponent->GetMorphTargetChangeTracker().GetOriginalMorphTargetName(Name);
+		if (OriginalMorphTargetName != NAME_None)
+		{
+			if (GUndo)
+			{
+				Mesh->SetFlags(RF_Transactional);
+				Mesh->Modify();
+			}
+
+			if (bAutoFill)
+			{
+				constexpr bool bRemoveZeroWeight = true;
+				Mesh->SetMorphTarget(OriginalMorphTargetName, 0.0f, bRemoveZeroWeight);	
+			}
+			else
+			{
+				constexpr bool bRemoveZeroWeight = false;
+				Mesh->SetMorphTarget(OriginalMorphTargetName, MorphTargetOverrides[Name], bRemoveZeroWeight);
+			}
+		}
+	}
+}
+
+void USkeletalMeshEditingCache::HandleMorphTargetEdited(FName MorphTarget)
+{
+	EditingMeshComponent->MarkMorphTargetEdited(MorphTarget);	
+}
+
+void USkeletalMeshEditingCache::OverrideMorphTargetWeight(FName MorphTarget, float Weight)
+{
+	HandleSetMorphTargetWeight(MorphTarget, Weight);
+}
+
+void USkeletalMeshEditingCache::ClearMorphTargetOverride(FName MorphTarget)
+{
+	if (MorphTarget != NAME_None)
+	{
+		constexpr bool bAutoFill = true;
+		constexpr float DummyWeight = 0.0f;
+		HandleSetMorphTargetAutoFill(MorphTarget, bAutoFill, DummyWeight);
+	}
+}
+
+FName USkeletalMeshEditingCache::AddMorphTarget(FName InName)
+{
+	return EditingMeshComponent->AddMorphTarget(InName);
+}
+
+FName USkeletalMeshEditingCache::RenameMorphTarget(FName InOldName, FName InNewName)
+{
+	FName MorphTarget = EditingMeshComponent->RenameMorphTarget(InOldName, InNewName);
+
+	if (float* OverrideWeight = MorphTargetOverrides.Find(InOldName))
+	{
+		float SavedWeight = *OverrideWeight;
+		ClearMorphTargetOverride(InOldName);
+		OverrideMorphTargetWeight(InNewName, SavedWeight);
+	}
+
+	return MorphTarget;
+}
+
+void USkeletalMeshEditingCache::RemoveMorphTargets(const TArray<FName>& InNames)
+{
+	EditingMeshComponent->RemoveMorphTargets(InNames);
+
+	for (const FName& Name : InNames)
+	{
+		ClearMorphTargetOverride(Name);
+	}
+}
+
+TArray<FName> USkeletalMeshEditingCache::DuplicateMorphTargets(const TArray<FName>& InNames)
+{
+	return EditingMeshComponent->DuplicateMorphTargets(InNames);
+}
+
+URefSkeletonPoser* USkeletalMeshEditingCache::GetSkeletonPoser() const
+{
+	return RefSkeletonPoser;
+}
+
+
+const TArray<FTransform>& USkeletalMeshEditingCache::GetComponentSpaceBoneTransformsRefPose() const
+{
+	return RefSkeletonPoser->GetComponentSpaceTransformsRefPose();
+}
+
+const TMap<FName, float>& USkeletalMeshEditingCache::GetSkeletalMeshComponentMorphTargetWeights() const
+{
+	UAnimInstance* AnimInstance = SkeletalMeshComponent->GetAnimInstance();
+	return AnimInstance->GetAnimationCurveList(EAnimCurveType::MorphTargetCurve);
+}
+
+void USkeletalMeshEditingCache::DeformPreviewMesh(const TArray<FTransform>& ComponentSpaceTransforms, const TMap<FName, float>& MorphTargetWeights)
+{
+	if (!PreviewMesh->IsVisible())
+	{
+		return;
+	}
+
+	if (ComponentSpaceTransforms.IsEmpty())
+	{
+		return;
+	}
+	
+	using namespace UE::Geometry;
+	using namespace SkeletalMeshToolsHelper;
+	
+	TArray<FMatrix> BoneMatrices = ComputeBoneMatrices(GetComponentSpaceBoneTransformsRefPose(), ComponentSpaceTransforms);
+			
+	FDynamicMesh3& EditingMesh = *EditingMeshComponent->GetMesh();
+				
+	auto DeformPreviewMeshFunc = [&](FDynamicMesh3& VisualMesh)
+		{
+			auto WriteFunc = [&](int32 VertID, const FVector& PosedVertPos)
+				{
+					VisualMesh.SetVertex(VertID, PosedVertPos);
+				};
+					
+			GetPosedMesh(WriteFunc, EditingMesh, BoneMatrices, NAME_None, MorphTargetWeights);
+			FMeshNormals::QuickRecomputeOverlayNormals(VisualMesh);
+		};
+		
+	constexpr bool bRebuildSpatial = false;
+	PreviewMesh->DeferredEditMesh(DeformPreviewMeshFunc, bRebuildSpatial);
+	PreviewMesh->NotifyDeferredEditCompleted(UPreviewMesh::ERenderUpdateMode::FastUpdate, EMeshRenderAttributeFlags::Positions | EMeshRenderAttributeFlags::VertexNormals, bRebuildSpatial);	
+}
+
+void USkeletalMeshEditingCache::UpdateSelectedBoneIndices()
+{
+	SelectedBoneIndices.Reset();
+	for (const FName Name: SelectedBoneNames)
+	{
+		SelectedBoneIndices.Add(GetEditingMeshComponent()->GetRefSkeleton().FindRawBoneIndex(Name));
+	}
+}
+
+void USkeletalMeshEditingCache::SetSkeletalMeshSkeletonDrawMode(ESkeletonDrawMode DrawMode) const
+{
+	if (UDebugSkelMeshComponent* DebugSkelMeshComponent = GetDebugSkelMeshComponent())
+	{
+		DebugSkelMeshComponent->SkeletonDrawMode = DrawMode;
+	}
+}
+
+ESkeletonDrawMode USkeletalMeshEditingCache::GetCurrentSkeletonDrawMode() const
+{
+	return CacheSkeletonDrawMode;
+}
+
+FLinearColor USkeletalMeshEditingCache::GetDefaultBoneColor(int32 InBoneIndex) const
+{
+	// this returns the normal unmodified color of the bone, calling code must account
+	// for any editor specific states that might affect the final bone color (like selection)
+	
+	// skeleton greyed out
+	if (GetCurrentSkeletonDrawMode() == ESkeletonDrawMode::GreyedOut)
+	{
+		return GetDefault<UPersonaOptions>()->DisabledBoneColor;
+	}
+
+	// using default color for all bones
+	if (!GetDefault<UPersonaOptions>()->bShowBoneColors)
+	{
+		return GetDefault<UPersonaOptions>()->DefaultBoneColor;
+	}
+	
+	// uses deterministic, semi-random desaturated color unique to the bone index
+	return SkeletalDebugRendering::GetSemiRandomColorForBone(InBoneIndex);	
+}
+
+UDebugSkelMeshComponent* USkeletalMeshEditingCache::GetDebugSkelMeshComponent() const
+{
+	return Cast<UDebugSkelMeshComponent>(SkeletalMeshComponent);
+}
+
+
+#undef LOCTEXT_NAMESPACE
diff --git a/Engine/Plugins/Animation/SkeletalMeshModelingTools/Source/SkeletalMeshModelingTools/Private/SkeletalMeshEditingCache.h b/Engine/Plugins/Animation/SkeletalMeshModelingTools/Source/SkeletalMeshModelingTools/Private/SkeletalMeshEditingCache.h
new file mode 100644
--- /dev/null
+++ b/Engine/Plugins/Animation/SkeletalMeshModelingTools/Source/SkeletalMeshModelingTools/Private/SkeletalMeshEditingCache.h
@@ -0,0 +1,198 @@
+﻿// Copyright Epic Games, Inc. All Rights Reserved.
+
+#pragma once
+
+#include "CoreMinimal.h"
+#include "SkeletalMeshNotifier.h"
+#include "Animation/DebugSkelMeshComponent.h"
+#include "Changes/ValueWatcher.h"
+#include "SkeletalMesh/SkeletalMeshToolsHelper.h"
+#include "TargetInterfaces/MeshTargetInterfaceTypes.h"
+
+#include "SkeletalMeshEditingCache.generated.h"
+
+#define UE_API SKELETALMESHMODELINGTOOLS_API
+
+class AActor;
+class USkeletalMeshBackedDynamicMeshComponent;
+class UPreviewMesh;
+class USkeletalMeshComponent;
+class URefSkeletonPoser;
+class FPrimitiveDrawInterface;
+struct FSkelDebugDrawConfig;
+class FEditorViewportClient;
+class FViewport;
+class FSceneView;
+class FCanvas;
+
+class USkeletalMeshEditingCache;
+class FSkeletalMeshEditingCacheNotifier: public ISkeletalMeshNotifier
+{
+public:
+	FSkeletalMeshEditingCacheNotifier(USkeletalMeshEditingCache* InEditingCahe);
+	virtual void HandleNotification(const TArray<FName>& BoneNames, const ESkeletalMeshNotifyType InNotifyType) override;
+private:
+	TWeakObjectPtr<USkeletalMeshEditingCache> EditingCache;
+};
+
+
+UCLASS(MinimalAPI, Transient)
+class USkeletalMeshEditingCache: public UObject
+{
+	GENERATED_BODY()
+	
+public:
+	DECLARE_DELEGATE_OneParam(FToggleSkeletalMeshBoneManipulation, bool);
+	DECLARE_DELEGATE_RetVal(bool, FIsSkeletalMeshBoneManipulationEnabled);
+	struct FDelegates
+	{
+		FToggleSkeletalMeshBoneManipulation ToggleSkeletalMeshBoneManipulationDelegate;
+		FIsSkeletalMeshBoneManipulationEnabled IsSkeletalMeshBoneManipulationEnabledDelegate;
+	};
+
+	void Spawn(UWorld* World, USkeletalMeshComponent* SkeletalMeshComponent, EMeshLODIdentifier LOD, const FDelegates& InDelegates);
+	EMeshLODIdentifier GetLOD() const;
+
+	void Destroy();
+
+	void ApplyChanges();
+	bool IsApplyingChanges() const;
+	void DiscardChanges();
+
+	void Tick();
+	bool HandleClick(HHitProxy *HitProxy);
+	void Render(FPrimitiveDrawInterface* PDI, TFunction<void(FSkelDebugDrawConfig&)> OverrideBoneDrawConfigFunc);
+	void DrawHUD(FEditorViewportClient* ViewportClient, FViewport* Viewport, const FSceneView* View, FCanvas* Canvas);
+
+	TSharedPtr<ISkeletalMeshNotifier> GetNotifier();
+	void HandleNotification(const TArray<FName>& BoneNames, const ESkeletalMeshNotifyType InNotifyType);
+
+	
+	USkeletalMeshBackedDynamicMeshComponent* GetEditingMeshComponent() const;	
+	const TArray<FTransform>& GetComponentSpaceBoneTransforms() const;
+	FTransform GetTransform();
+
+	void HideSkeleton();
+	void ShowSkeleton();
+	void ToggleBoneManipulation(bool bEnable);
+	bool IsDynamicMeshSkeletonEnabled() const;
+	bool IsDynamicMeshBoneManipulationEnabled() const;
+	int32 GetFirstSelectedBoneIndex() const;
+	TArray<FName> GetSelectedBones() const;
+	void ResetDynamicMeshBoneTransforms(bool bSelectedOnly);
+	
+	TArray<FName> GetMorphTargets() const;
+	TMap<FName, float> GetMorphTargetWeights() const;
+	float GetMorphTargetWeight(FName MorphTarget) const;
+
+	void HandleSetMorphTargetWeight(FName MorphTarget, float Weight);
+	bool GetMorphTargetAutoFill(FName MorphTarget);
+	void HandleSetMorphTargetAutoFill(FName MorphTarget, bool bAutoFill, float PreviousOverrideWeight);
+	void HandleMorphTargetEdited(FName MorphTarget);
+
+	void OverrideMorphTargetWeight(FName MorphTarget, float Weight);
+	void ClearMorphTargetOverride(FName MorphTarget);
+
+	FName AddMorphTarget(FName InName);
+	FName RenameMorphTarget(FName InOldName, FName InNewName);
+	void RemoveMorphTargets(const TArray<FName>& InNames);
+	TArray<FName> DuplicateMorphTargets(const TArray<FName>& InNames);
+	
+	URefSkeletonPoser* GetSkeletonPoser() const;
+	
+protected:
+	UPROPERTY()
+	TObjectPtr<AActor> HostActor;
+
+	UPROPERTY()
+	TObjectPtr<USkeletalMeshBackedDynamicMeshComponent> EditingMeshComponent;
+
+	UPROPERTY()
+	TObjectPtr<UPreviewMesh> PreviewMesh;
+
+	UPROPERTY()
+	TWeakObjectPtr<USkeletalMeshComponent> SkeletalMeshComponent;
+
+	FDelegates Delegates;
+
+	// Flags that drives both dyna mesh and skeletal mesh
+	bool bCacheVisibility = true;
+	ESkeletonDrawMode CacheSkeletonDrawMode = ESkeletonDrawMode::Default;
+	bool bCacheBoneManipulation = true;
+
+	bool bEnableDynamicMesh = false;
+	bool bEnableDynamicMeshSkeleton = false;
+	
+	TValueWatcher<int32> ChangeCountWatcher;
+	TValueWatcher<int32> SkeletonChangeCountWatcher;
+
+	struct FMeshVisibilityState
+	{
+		bool bEnableDynamicMesh= false;
+		bool bCacheVisibility = true;
+		friend bool operator==(const FMeshVisibilityState&, const FMeshVisibilityState&) = default;
+	};
+
+	TValueWatcher<FMeshVisibilityState> MeshVisibilityWatcher;
+	
+	struct FSkeletonVisibilityState
+	{
+		bool bEnableDynamicMeshSkeleton = false;
+		ESkeletonDrawMode CacheSkeletonDrawMode = ESkeletonDrawMode::Default;
+		friend bool operator==(const FSkeletonVisibilityState&, const FSkeletonVisibilityState&) = default;
+	};
+	
+	TValueWatcher<FSkeletonVisibilityState> SkeletonVisibilityWatcher;
+
+	struct FBoneManipulationState
+	{
+		bool bEnableDynamicMeshSkeleton = false;
+		bool bCacheBoneManipulation = true;
+		friend bool operator==(const FBoneManipulationState&, const FBoneManipulationState&) = default;
+	};
+
+	
+	TValueWatcher<FBoneManipulationState> BoneManipulationWatcher;
+	
+	TValueWatcher<bool> PreviewMeshVisibilityWatcher;
+
+	void HandleVisibilityChangeRequest(bool bVisible);
+
+	UPROPERTY()
+	TObjectPtr<URefSkeletonPoser> RefSkeletonPoser;
+
+	UPROPERTY()
+	TMap<FName, float> MorphTargetOverrides;
+
+	SkeletalMeshToolsHelper::FPoseChangeDetector PoseChangeDetector;
+	
+	void HandlePoseChangeDetectorEvent(SkeletalMeshToolsHelper::FPoseChangeDetector::FPayload Payload);
+
+	const TArray<FTransform>& GetComponentSpaceBoneTransformsRefPose() const;
+
+	const TMap<FName, float>& GetSkeletalMeshComponentMorphTargetWeights() const;
+
+	void DeformPreviewMesh(const TArray<FTransform>& ComponentSpaceTransforms, const TMap<FName, float>& MorphTargetWeights);
+	
+	TSharedPtr<FSkeletalMeshEditingCacheNotifier> Notifier;
+	
+	TArray<FName> SelectedBoneNames;
+	TArray<int32> SelectedBoneIndices;
+	void UpdateSelectedBoneIndices();
+
+	
+	ESkeletonDrawMode DynamicMeshSkeletonDrawMode = ESkeletonDrawMode::Hidden;
+	
+	void SetSkeletalMeshSkeletonDrawMode(ESkeletonDrawMode DrawMode) const;
+	ESkeletonDrawMode GetCurrentSkeletonDrawMode() const;
+	
+	ESkeletonDrawMode SavedSkeletonDrawMode = ESkeletonDrawMode::Default;
+
+	FLinearColor GetDefaultBoneColor(int32 BoneIndex) const;
+
+	UDebugSkelMeshComponent* GetDebugSkelMeshComponent() const;
+
+};
+
+
+#undef UE_API
diff --git a/Engine/Plugins/Animation/SkeletalMeshModelingTools/Source/SkeletalMeshModelingTools/Private/SkeletalMeshEditorUtils.cpp b/Engine/Plugins/Animation/SkeletalMeshModelingTools/Source/SkeletalMeshModelingTools/Private/SkeletalMeshEditorUtils.cpp
--- a/Engine/Plugins/Animation/SkeletalMeshModelingTools/Source/SkeletalMeshModelingTools/Private/SkeletalMeshEditorUtils.cpp
+++ b/Engine/Plugins/Animation/SkeletalMeshModelingTools/Source/SkeletalMeshModelingTools/Private/SkeletalMeshEditorUtils.cpp
@@ -6,6 +6,7 @@
 #include "EditorModeManager.h"
 #include "InteractiveToolsContext.h"
 #include "ISkeletonEditorModule.h"
+#include "SkeletalMeshEditingCache.h"
 #include "SkeletalMeshModelingToolsEditorMode.h"
 #include "SkeletonModifier.h"
 #include "SReferenceSkeletonTree.h"
@@ -53,6 +54,47 @@ USkeletalMeshEditorContextObject* UE::SkeletalMeshEditorUtils::GetEditorContextO
 	return ToolsContext->ContextObjectStore->FindContext<USkeletalMeshEditorContextObject>();
 }
 
+UE::SkeletalMeshEditorUtils::FSkeletalMeshNotifierBindScope::FSkeletalMeshNotifierBindScope(
+	TWeakPtr<ISkeletalMeshNotifier> InNotifierA,
+	TWeakPtr<ISkeletalMeshNotifier> InNotifierB
+	) : NotifierA(InNotifierA)
+	, NotifierB(InNotifierB)
+{
+
+	DelegateToRemoveFromNotifierA = InNotifierA.Pin()->Delegate().AddLambda(
+	[InNotifierB](const TArray<FName>& BoneNames, const ESkeletalMeshNotifyType InNotifyType)
+	{
+		if (InNotifierB.IsValid())
+		{
+			InNotifierB.Pin()->HandleNotification(BoneNames, InNotifyType);
+		}
+	});
+	
+	DelegateToRemoveFromNotifierB = InNotifierB.Pin()->Delegate().AddLambda(
+	[InNotifierA](const TArray<FName>& BoneNames, const ESkeletalMeshNotifyType InNotifyType)
+	{
+		if (InNotifierA.IsValid())
+		{
+			InNotifierA.Pin()->HandleNotification(BoneNames, InNotifyType);
+		}
+	});
+}
+
+UE::SkeletalMeshEditorUtils::FSkeletalMeshNotifierBindScope::~FSkeletalMeshNotifierBindScope()
+{
+	if (DelegateToRemoveFromNotifierA.IsValid() && NotifierA.IsValid())
+	{
+		check(NotifierA.Pin()->Delegate().Remove(DelegateToRemoveFromNotifierA));
+	}
+
+	if (DelegateToRemoveFromNotifierB.IsValid() && NotifierB.IsValid())
+	{
+		check(NotifierB.Pin()->Delegate().Remove(DelegateToRemoveFromNotifierB));
+	}
+}
+
+
+
 void USkeletalMeshEditorContextObject::Register(UInteractiveToolManager* InToolManager)
 {
 	if (ensure(!bRegistered) == false)
@@ -106,11 +148,24 @@ void USkeletalMeshEditorContextObject::ToggleBoneManipulation(bool bEnable)
 
 const TArray<FTransform>& USkeletalMeshEditorContextObject::GetComponentSpaceBoneTransforms(UToolTarget* InToolTarget)
 {
-	return EditorMode->GetComponentSpaceBoneTransforms();
+	return EditorMode->GetCurrentEditingCache()->GetComponentSpaceBoneTransforms();
 }
 
 
+FName USkeletalMeshEditorContextObject::GetEditingMorphTarget()
+{
+	return EditorMode->GetEditingMorphTarget();
+}
+
+TMap<FName, float> USkeletalMeshEditorContextObject::GetMorphTargetWeights()
+{
+	return EditorMode->GetCurrentEditingCache()->GetMorphTargetWeights();
+}
 
+void USkeletalMeshEditorContextObject::NotifyMorphTargetEdited()
+{
+	return EditorMode->GetCurrentEditingCache()->HandleMorphTargetEdited( EditorMode->GetEditingMorphTarget() );
+}
 
 void USkeletalMeshEditorContextObject::BindTo(ISkeletalMeshEditingInterface* InEditingInterface)
 {
@@ -134,41 +189,23 @@ void USkeletalMeshEditorContextObject::UnbindFrom(ISkeletalMeshEditingInterface*
 	UnbindRefSkeletonTree(InEditingInterface);
 }
 
-TPair<FDelegateHandle, FDelegateHandle> USkeletalMeshEditorContextObject::BindInterfaceTo(
+
+
+USkeletalMeshEditorContextObject::FBindData USkeletalMeshEditorContextObject::BindInterfaceTo(
 	ISkeletalMeshEditingInterface* InInterface,
-	ISkeletalMeshNotifier& InOtherNotifier)
+	TSharedPtr<ISkeletalMeshNotifier> InOtherNotifier)
 {
-	// connect external interface to tool (ie skeletal mesh editor -> tool)
-	FDelegateHandle ToToolNotifierHandle = InOtherNotifier.Delegate().AddLambda(
-	[InInterface](const TArray<FName>& BoneNames, const ESkeletalMeshNotifyType InNotifyType)
-	{
-		InInterface->GetNotifier().HandleNotification(BoneNames, InNotifyType);
-	});
-
-	// connect tool to external interface (ie tool -> skeletal mesh editor)
-	FDelegateHandle FromToolNotifierHandle = InInterface->GetNotifier().Delegate().AddLambda(
-		[&InOtherNotifier](const TArray<FName>& BoneNames, const ESkeletalMeshNotifyType InNotifyType)
-	{
-		InOtherNotifier.HandleNotification(BoneNames, InNotifyType);
-	});
+	FBindData BindData;
+	BindData.BindScope.Reset(new UE::SkeletalMeshEditorUtils::FSkeletalMeshNotifierBindScope(InInterface->GetNotifier(), InOtherNotifier));
+	
+	return BindData;
 
-	return {ToToolNotifierHandle, FromToolNotifierHandle};
 }
 
 void USkeletalMeshEditorContextObject::UnbindInterfaceFrom(
-	ISkeletalMeshEditingInterface* InInterface,
-	ISkeletalMeshNotifier& InOtherNotifier,
-	const FBindData& InOutBindData)
+	FBindData& InOutBindData)
 {
-	if (InOutBindData.ToToolNotifierHandle.IsValid())
-	{
-		InOtherNotifier.Delegate().Remove(InOutBindData.ToToolNotifierHandle);
-	}
-
-	if (InOutBindData.FromToolNotifierHandle.IsValid())
-	{
-		InInterface->GetNotifier().Delegate().Remove(InOutBindData.FromToolNotifierHandle);
-	}
+	InOutBindData.BindScope.Reset();
 }
 
 void USkeletalMeshEditorContextObject::BindEditor(ISkeletalMeshEditingInterface* InEditingInterface)
@@ -183,35 +220,22 @@ void USkeletalMeshEditorContextObject::BindEditor(ISkeletalMeshEditingInterface*
 		return;
 	}
 
-	TSharedPtr<ISkeletalMeshEditorBinding> Binding = EditorMode.Pin()->GetEditorBinding();
+	TSharedPtr<ISkeletalMeshEditorBinding> Binding = EditorMode.Pin()->GetModeBinding();
 	if (!Binding.IsValid())
 	{
 		return;
 	}
+	
+	EditorBindings.Emplace(InEditingInterface, BindInterfaceTo(InEditingInterface, Binding->GetNotifier()));
 
-	InEditingInterface->BindTo(Binding);
-
-	FBindData BindData;
-	Tie(BindData.ToToolNotifierHandle, BindData.FromToolNotifierHandle) = BindInterfaceTo(InEditingInterface, Binding->GetNotifier());
-	EditorBindings.Emplace(InEditingInterface, BindData);
-
-	InEditingInterface->GetNotifier().HandleNotification(Binding->GetSelectedBones(), ESkeletalMeshNotifyType::BonesSelected);
+	InEditingInterface->GetNotifier()->HandleNotification(Binding->GetSelectedBones(), ESkeletalMeshNotifyType::BonesSelected);
 }
 
 void USkeletalMeshEditorContextObject::UnbindEditor(ISkeletalMeshEditingInterface* InEditingInterface)
 {
-	if (const FBindData* BindData = EditorBindings.Find(InEditingInterface))
+	if (FBindData* BindData = EditorBindings.Find(InEditingInterface))
 	{
-		if (BindData->FromToolNotifierHandle.IsValid())
-		{
-			InEditingInterface->Unbind();
-		}
-
-		if (EditorMode.IsValid())
-		{
-			const TSharedPtr<ISkeletalMeshEditorBinding> Binding = EditorMode.Pin()->GetEditorBinding();
-			UnbindInterfaceFrom(InEditingInterface, Binding->GetNotifier(), *BindData);
-		}
+		BindData->BindScope.Reset();
 	
 		EditorBindings.Remove(InEditingInterface);
 	}
@@ -254,34 +278,34 @@ void USkeletalMeshEditorContextObject::BindRefSkeletonTree(ISkeletalMeshEditingI
 	// switch between SSkeletonTree and SReferenceSkeletonTree
 	SkeletonTab->SetContent(RefSkeletonWidget.ToSharedRef());
 
-	FBindData BindData;
-	Tie(BindData.ToToolNotifierHandle, BindData.FromToolNotifierHandle) = BindInterfaceTo(InEditingInterface, RefSkeletonTree->GetNotifier());
-	TreeBindings.Emplace(InEditingInterface, BindData);
+	TreeBindings.Emplace(InEditingInterface, BindInterfaceTo(InEditingInterface, RefSkeletonTree->GetNotifier()));
 
 	// unbind notifications on the initial SSkeletonTree to avoid useless notifications
-	if(const TSharedPtr<ISkeletalMeshEditorBinding> Binding = EditorMode->GetEditorBinding())
+	if(const TSharedPtr<ISkeletalMeshEditorBinding> Binding = EditorMode->GetModeBinding())
 	{
 		const TArray<FName> SelectedBones = Binding->GetSelectedBones();
 		
-		if (const FBindData* EditorBindData = EditorBindings.Find(InEditingInterface))
+		if (FBindData* EditorBindData = EditorBindings.Find(InEditingInterface))
 		{
-			UnbindInterfaceFrom(InEditingInterface, Binding->GetNotifier(), *EditorBindData);
-			Binding->GetNotifier().HandleNotification({}, ESkeletalMeshNotifyType::BonesSelected);
+			UnbindInterfaceFrom(*EditorBindData);
+			EditorBindings.Remove(InEditingInterface);
+			
+			Binding->GetNotifier()->HandleNotification({}, ESkeletalMeshNotifyType::BonesSelected);
 		}
 
-		RefSkeletonTree->GetNotifier().HandleNotification(SelectedBones, ESkeletalMeshNotifyType::BonesSelected);
+		RefSkeletonTree->GetNotifier()->HandleNotification(SelectedBones, ESkeletalMeshNotifyType::BonesSelected);
 	}
 }
 
 void USkeletalMeshEditorContextObject::UnbindRefSkeletonTree(ISkeletalMeshEditingInterface* InEditingInterface)
 {
-	if (const FBindData* BindData = TreeBindings.Find(InEditingInterface))
+	if (FBindData* BindData = TreeBindings.Find(InEditingInterface))
 	{
 		TArray<FName> SelectedBones;
 		if (RefSkeletonTree.IsValid())
 		{
 			RefSkeletonTree->GetSelectedBoneNames(SelectedBones);
-			UnbindInterfaceFrom(InEditingInterface, RefSkeletonTree->GetNotifier(), *BindData);
+			UnbindInterfaceFrom(*BindData);
 		}
 
 		if (EditorMode.IsValid())
@@ -299,12 +323,12 @@ void USkeletalMeshEditorContextObject::UnbindRefSkeletonTree(ISkeletalMeshEditin
 				}	
 			}
 
-			// update the SSkeletonTree with the new RefSkeleton and set current selection 
-			if (const TSharedPtr<ISkeletalMeshEditorBinding> Binding = EditorMode->GetEditorBinding())
+			// update the default skeleton widget with the new RefSkeleton and set current selection 
+			if (const TSharedPtr<ISkeletalMeshEditorBinding> Binding = EditorMode->GetModeBinding())
 			{
-				ISkeletalMeshNotifier& Notifier = Binding->GetNotifier();
-				Notifier.HandleNotification(SelectedBones, ESkeletalMeshNotifyType::HierarchyChanged);
-				Notifier.HandleNotification(SelectedBones, ESkeletalMeshNotifyType::BonesSelected);
+				TSharedPtr<ISkeletalMeshNotifier> Notifier = Binding->GetNotifier();
+				Notifier->HandleNotification(SelectedBones, ESkeletalMeshNotifyType::HierarchyChanged);
+				Notifier->HandleNotification(SelectedBones, ESkeletalMeshNotifyType::BonesSelected);
 			}
 		}
 		
@@ -323,5 +347,5 @@ const FName& USkeletalMeshEditorContextObject::GetSkeletonTreeTabId()
 
 TSharedPtr<ISkeletalMeshEditorBinding> USkeletalMeshEditorContextObject::GetBinding() const
 {
-	return EditorMode.IsValid() ? EditorMode->GetEditorBinding() : nullptr;
+	return EditorMode.IsValid() ? EditorMode->GetModeBinding() : nullptr;
 }
diff --git a/Engine/Plugins/Animation/SkeletalMeshModelingTools/Source/SkeletalMeshModelingTools/Private/SkeletalMeshEditorUtils.h b/Engine/Plugins/Animation/SkeletalMeshModelingTools/Source/SkeletalMeshModelingTools/Private/SkeletalMeshEditorUtils.h
--- a/Engine/Plugins/Animation/SkeletalMeshModelingTools/Source/SkeletalMeshModelingTools/Private/SkeletalMeshEditorUtils.h
+++ b/Engine/Plugins/Animation/SkeletalMeshModelingTools/Source/SkeletalMeshModelingTools/Private/SkeletalMeshEditorUtils.h
@@ -27,6 +27,19 @@ namespace SkeletalMeshEditorUtils
 	UE_API bool UnregisterEditorContextObject(UInteractiveToolsContext* InToolsContext);
 
 	UE_API USkeletalMeshEditorContextObject* GetEditorContextObject(UInteractiveToolsContext* InToolsContext);
+
+
+	struct FSkeletalMeshNotifierBindScope 
+	{
+		FSkeletalMeshNotifierBindScope(TWeakPtr<ISkeletalMeshNotifier> InNotifierA,  TWeakPtr<ISkeletalMeshNotifier> InNotifierB);
+		~FSkeletalMeshNotifierBindScope();
+		
+		TWeakPtr<ISkeletalMeshNotifier> NotifierA;
+		TWeakPtr<ISkeletalMeshNotifier> NotifierB;
+
+		FDelegateHandle	DelegateToRemoveFromNotifierA;	
+		FDelegateHandle	DelegateToRemoveFromNotifierB;	
+	};
 }
 	
 }
@@ -53,6 +66,10 @@ public:
 	UE_API virtual void ToggleBoneManipulation(bool bEnable) override;
 	UE_API virtual const TArray<FTransform>& GetComponentSpaceBoneTransforms(UToolTarget* InToolTarget) override;
 
+	UE_API virtual FName GetEditingMorphTarget() override;
+	UE_API virtual TMap<FName, float> GetMorphTargetWeights() override;
+	UE_API virtual void NotifyMorphTargetEdited() override;
+	
 
 	UE_API virtual void BindTo(ISkeletalMeshEditingInterface* InEditingInterface) override;
 	UE_API virtual void UnbindFrom(ISkeletalMeshEditingInterface* InEditingInterface) override;
@@ -67,20 +84,17 @@ private:
 
 	struct FBindData
 	{
-		FDelegateHandle ToToolNotifierHandle;
-		FDelegateHandle FromToolNotifierHandle;		
+		TUniquePtr<UE::SkeletalMeshEditorUtils::FSkeletalMeshNotifierBindScope> BindScope;
 	};
 	TMap< ISkeletalMeshEditingInterface*, FBindData > EditorBindings;
 	TMap< ISkeletalMeshEditingInterface*, FBindData > TreeBindings;
 
-	static UE_API TPair<FDelegateHandle, FDelegateHandle> BindInterfaceTo(
+	static UE_API FBindData BindInterfaceTo(
 		ISkeletalMeshEditingInterface* InInterface,
-		ISkeletalMeshNotifier& InOtherNotifier);
+		TSharedPtr<ISkeletalMeshNotifier> InOtherNotifier);
 	
 	static UE_API void UnbindInterfaceFrom(
-		ISkeletalMeshEditingInterface* InInterface,
-		ISkeletalMeshNotifier& InOtherNotifier,
-		const FBindData& InOutBindData);
+		FBindData& InOutBindData);
 
 	UE_API void BindEditor(ISkeletalMeshEditingInterface* InEditingInterface);
 	UE_API void UnbindEditor(ISkeletalMeshEditingInterface* InEditingInterface);
diff --git a/Engine/Plugins/Animation/SkeletalMeshModelingTools/Source/SkeletalMeshModelingTools/Private/SkeletalMeshModelingToolsCommands.h b/Engine/Plugins/Animation/SkeletalMeshModelingTools/Source/SkeletalMeshModelingTools/Private/SkeletalMeshModelingToolsCommands.h
--- a/Engine/Plugins/Animation/SkeletalMeshModelingTools/Source/SkeletalMeshModelingTools/Private/SkeletalMeshModelingToolsCommands.h
+++ b/Engine/Plugins/Animation/SkeletalMeshModelingTools/Source/SkeletalMeshModelingTools/Private/SkeletalMeshModelingToolsCommands.h
@@ -39,6 +39,9 @@ public:
 	TSharedPtr<FUICommandInfo> CopyBones;
 	TSharedPtr<FUICommandInfo> PasteBones;
 	TSharedPtr<FUICommandInfo> DuplicateBones;
+
+	// Morph Target Editing commands
+	TSharedPtr<FUICommandInfo> NewMorphTarget;
 };
 
 class FSkeletalMeshModelingToolsActionCommands : public TInteractiveToolCommands<FSkeletalMeshModelingToolsActionCommands>
diff --git a/Engine/Plugins/Animation/SkeletalMeshModelingTools/Source/SkeletalMeshModelingTools/Private/SkeletalMeshModelingToolsEditorMode.cpp b/Engine/Plugins/Animation/SkeletalMeshModelingTools/Source/SkeletalMeshModelingTools/Private/SkeletalMeshModelingToolsEditorMode.cpp
--- a/Engine/Plugins/Animation/SkeletalMeshModelingTools/Source/SkeletalMeshModelingTools/Private/SkeletalMeshModelingToolsEditorMode.cpp
+++ b/Engine/Plugins/Animation/SkeletalMeshModelingTools/Source/SkeletalMeshModelingTools/Private/SkeletalMeshModelingToolsEditorMode.cpp
@@ -17,6 +17,9 @@
 // Stylus support is currently disabled due to issues with the stylus plugin
 // We are leaving the code in this cpp file, defined out, so that it is easier to bring back if/when the stylus plugin is improved.
 
+#include "SkeletalMeshModelingToolsModule.h"
+#include "Widgets/Notifications/SNotificationList.h"
+
 #include UE_INLINE_GENERATED_CPP_BY_NAME(SkeletalMeshModelingToolsEditorMode)
 #define ENABLE_STYLUS_SUPPORT 0
 
@@ -25,6 +28,7 @@
 #include "IStylusState.h"
 #endif
 
+#include "AnimationEditorViewportClient.h"
 #include "ToolTargets/SkeletalMeshComponentToolTarget.h"
 #include "Components/SkeletalMeshComponent.h"
 
@@ -35,6 +39,7 @@
 #include "EditMeshPolygonsTool.h"
 #include "EditorInteractiveGizmoManager.h"
 #include "HoleFillTool.h"
+#include "IPersonaPreviewScene.h"
 #include "IPersonaToolkit.h"
 #include "ISkeletalMeshEditor.h"
 #include "LatticeDeformerTool.h"
@@ -51,6 +56,7 @@
 #include "RemeshMeshTool.h"
 #include "RemoveOccludedTrianglesTool.h"
 #include "SimplifyMeshTool.h"
+#include "SkeletalMeshEditingCache.h"
 #include "SkeletalMeshEditorUtils.h"
 #include "SmoothMeshTool.h"
 #include "ToolTargetManager.h"
@@ -58,12 +64,15 @@
 
 #include "SkeletalMeshNotifier.h"
 #include "Animation/DebugSkelMeshComponent.h"
+#include "Components/SKMBackedDynaMeshComponent.h"
+#include "SkeletalMesh/RefSkeletonPoser.h"
 
 #include "SkeletalMesh/SkeletalMeshEditingInterface.h"
 #include "SkeletalMesh/SkeletonEditingTool.h"
 #include "SkeletalMesh/SkinWeightsBindingTool.h"
 #include "SkeletalMesh/SkinWeightsPaintTool.h"
 #include "ToolTargets/SkeletalMeshToolTarget.h"
+#include "ToolTargets/SKMBackedDynaMeshComponentToolTarget.h"
 
 #define LOCTEXT_NAMESPACE "SkeletalMeshModelingToolsEditorMode"
 
@@ -155,12 +164,75 @@ public:
 };
 #endif // ENABLE_STYLUS_SUPPORT
 
+static void ShowEditorMessage(ELogVerbosity::Type InMessageType, const FText& InMessage)
+{
+	FNotificationInfo Notification(InMessage);
+	Notification.bUseSuccessFailIcons = true;
+	Notification.ExpireDuration = 5.0f;
+
+	SNotificationItem::ECompletionState State = SNotificationItem::CS_Success;
+
+	switch(InMessageType)
+	{
+	case ELogVerbosity::Warning:
+		UE_LOG(LogSkeletalMeshModelingTools, Warning, TEXT("%s"), *InMessage.ToString());
+		break;
+	case ELogVerbosity::Error:
+		State = SNotificationItem::CS_Fail;
+		UE_LOG(LogSkeletalMeshModelingTools, Error, TEXT("%s"), *InMessage.ToString());
+		break;
+	default:
+		break; // don't log anything unless a warning or error
+	}
+	
+	FSlateNotificationManager::Get().AddNotification(Notification)->SetCompletionState(State);
+}
 
 // NOTE: This is a simple proxy at the moment. In the future we want to pull in more of the 
 // modeling tools as we add support in the skelmesh storage.
 
 const FEditorModeID USkeletalMeshModelingToolsEditorMode::Id("SkeletalMeshModelingToolsEditorMode");
 
+FSkeletalMeshModelingToolsEditorModeNotifier::FSkeletalMeshModelingToolsEditorModeNotifier(USkeletalMeshModelingToolsEditorMode* InEditorMode)
+	:EditorMode(InEditorMode)
+{
+}
+
+void FSkeletalMeshModelingToolsEditorModeNotifier::HandleNotification(const TArray<FName>& BoneNames, const ESkeletalMeshNotifyType InNotifyType)
+{
+	if (Notifying())
+	{
+		return;
+	}
+
+	Notify(BoneNames, InNotifyType);
+}
+
+FSkeletalMeshModelingToolsEditorModeBinding::FSkeletalMeshModelingToolsEditorModeBinding(USkeletalMeshModelingToolsEditorMode* InEditorMode)
+	: EditorMode(InEditorMode)
+{
+}
+
+TSharedPtr<ISkeletalMeshNotifier> FSkeletalMeshModelingToolsEditorModeBinding::GetNotifier()
+{
+	if (!Notifier)
+	{
+		Notifier = MakeShared<FSkeletalMeshModelingToolsEditorModeNotifier>(EditorMode.Get());
+	}
+	
+	return Notifier;
+}
+
+ISkeletalMeshEditorBinding::NameFunction FSkeletalMeshModelingToolsEditorModeBinding::GetNameFunction()
+{
+	// unused;
+	return {};
+}
+
+TArray<FName> FSkeletalMeshModelingToolsEditorModeBinding::GetSelectedBones() const
+{
+	return EditorMode->GetSelectedBones();
+}
 
 USkeletalMeshModelingToolsEditorMode::USkeletalMeshModelingToolsEditorMode() 
 {
@@ -199,8 +271,13 @@ void USkeletalMeshModelingToolsEditorMode::Enter()
 
 	if (TObjectPtr<UToolTargetManager> ToolTargetManager = InteractiveToolsContext->TargetManager)
 	{
-		ToolTargetManager->AddTargetFactory( NewObject<USkeletalMeshComponentToolTargetFactory>(ToolTargetManager) );
-		ToolTargetManager->AddTargetFactory( NewObject<USkeletalMeshReadOnlyToolTargetFactory>(ToolTargetManager) );
+		// Special tool target factory to support fast tool switch for supported tools
+		USkeletalMeshBackedDynamicMeshComponentToolTargetFactory* SKMDynaMeshFactory = NewObject<USkeletalMeshBackedDynamicMeshComponentToolTargetFactory>(ToolTargetManager);
+		SKMDynaMeshFactory->Init(this);
+		
+		ToolTargetManager->AddTargetFactory(SKMDynaMeshFactory);
+		ToolTargetManager->AddTargetFactory( NewObject<USkeletalMeshReadOnlyToolTargetFactory>(ToolTargetManager) );	
+	
 	}
 
 #if ENABLE_STYLUS_SUPPORT
@@ -275,6 +352,66 @@ void USkeletalMeshModelingToolsEditorMode::Enter()
 
 	GetInteractiveToolsContext()->ToolManager->OnToolPostBuild.AddUObject(this, &USkeletalMeshModelingToolsEditorMode::OnToolPostBuild);
 	GetInteractiveToolsContext()->ToolManager->OnToolEndedWithStatus.AddUObject(this, &USkeletalMeshModelingToolsEditorMode::OnToolEndedWithStatus);
+
+	CacheChangeCountWatcher.Initialize(
+		[this]()
+			{
+				if (GetCurrentEditingCache())
+				{
+					return GetCurrentEditingCache()->GetEditingMeshComponent()->GetChangeCount();
+				}
+				return 0;
+			},
+		[&](int32 ChangeCount)
+			{
+				TypedToolkit.Pin()->RefreshMorphTargetManager();
+			},
+		0);
+	
+	CacheDynamicMeshSkeletonStatusWatcher.Initialize(
+		[this]()
+			{
+				if (GetCurrentEditingCache())
+				{
+					return GetCurrentEditingCache()->IsDynamicMeshSkeletonEnabled();
+				}
+				return false;
+			},
+		[this](bool bEnableDynamicMeshSkeleton)
+			{
+				if (TypedToolkit.IsValid())
+				{
+					if (bEnableDynamicMeshSkeleton)
+					{
+						TypedToolkit.Pin()->ShowDynamicMeshSkeletonTree();
+					}
+					else
+					{
+						TypedToolkit.Pin()->ShowSkeletalMeshSkeletonTree();
+					}
+				}
+			},
+		false);
+	
+	CacheSkeletonChangeCountWatcher.Initialize(
+		[this]()
+			{
+				if (GetCurrentEditingCache())
+				{
+					return GetCurrentEditingCache()->GetEditingMeshComponent()->GetSkeletonChangeTracker().GetChangeCount();
+				}
+				return 0;
+			},
+		[this](int32 ChangeCount)
+			{
+				if (GetCurrentEditingCache())
+				{
+					SkeletonReader->ExternalUpdate(GetCurrentEditingCache()->GetEditingMeshComponent()->GetRefSkeleton());
+					GetModeBinding()->GetNotifier()->Notify({}, ESkeletalMeshNotifyType::HierarchyChanged);
+				}
+			},
+		0);
+	
 }
 
 UDebugSkelMeshComponent* USkeletalMeshModelingToolsEditorMode::GetSkelMeshComponent() const
@@ -288,63 +425,111 @@ UDebugSkelMeshComponent* USkeletalMeshModelingToolsEditorMode::GetSkelMeshCompon
 	return Cast<UDebugSkelMeshComponent>(SkeletalMeshComponent);
 }
 
-ESkeletonDrawMode USkeletalMeshModelingToolsEditorMode::GetSkeletalMeshSkeletonDrawMode() const
+
+TSharedPtr<FTabManager> USkeletalMeshModelingToolsEditorMode::GetAssociatedTabManager()
 {
-	if (UDebugSkelMeshComponent* SkeletalMeshComponent = GetSkelMeshComponent())
+	if (Editor.IsValid())
 	{
-		return SkeletalMeshComponent->SkeletonDrawMode;
+		return Editor.Pin()->GetAssociatedTabManager();
 	}
 
-	return ESkeletonDrawMode::Default;
+	return nullptr;
 }
 
-void USkeletalMeshModelingToolsEditorMode::SetSkeletalMeshSkeletonDrawMode(ESkeletonDrawMode DrawMode) const
+const TArray<FTransform>& USkeletalMeshModelingToolsEditorMode::GetComponentSpaceBoneTransforms() const
 {
-	if (UDebugSkelMeshComponent* SkeletalMeshComponent = GetSkelMeshComponent())
-	{
-		SkeletalMeshComponent->SkeletonDrawMode = DrawMode;
-	}
+	return GetCurrentEditingCache()->GetComponentSpaceBoneTransforms();
 }
 
-ESkeletonDrawMode USkeletalMeshModelingToolsEditorMode::GetCurrentSkeletonDrawMode() const
+void USkeletalMeshModelingToolsEditorMode::ToggleBoneManipulation(bool bEnable)
 {
-	return GetSkeletalMeshSkeletonDrawMode();
+	GetCurrentEditingCache()->ToggleBoneManipulation(bEnable);
 }
 
+USkeletonModifier* USkeletalMeshModelingToolsEditorMode::GetSkeletonReader()
+{
+	return SkeletonReader;
+}
 
-TSharedPtr<FTabManager> USkeletalMeshModelingToolsEditorMode::GetAssociatedTabManager()
+TArray<FName> USkeletalMeshModelingToolsEditorMode::GetSelectedBones() const 
 {
-	if (Editor.IsValid())
+	return GetCurrentEditingCache()->GetSelectedBones();
+}
+
+
+FName USkeletalMeshModelingToolsEditorMode::GetEditingMorphTarget()
+{
+	if (!GetCurrentEditingCache()->GetMorphTargets().Contains(EditingMorphTarget))
 	{
-		return Editor.Pin()->GetAssociatedTabManager();
+		return NAME_None;
 	}
 
-	return nullptr;
+	return EditingMorphTarget;
 }
 
-const TArray<FTransform>& USkeletalMeshModelingToolsEditorMode::GetComponentSpaceBoneTransforms() const
+void USkeletalMeshModelingToolsEditorMode::HandleSetEditingMorphTarget(FName InMorphTarget)
 {
-	return GetSkelMeshComponent()->GetComponentSpaceTransforms();
+	FName OldEditingMorphTarget = EditingMorphTarget;
+	EditingMorphTarget = InMorphTarget;
+
+	GetCurrentEditingCache()->OverrideMorphTargetWeight(EditingMorphTarget, 1.0f);
+	GetCurrentEditingCache()->ClearMorphTargetOverride(OldEditingMorphTarget);
 }
 
-void USkeletalMeshModelingToolsEditorMode::ToggleBoneManipulation(bool bEnable)
+
+
+TArray<FName> USkeletalMeshModelingToolsEditorMode::GetMorphTargets()
 {
-	if (Owner)
-	{
-		if (bEnable)
-		{
-			Owner->ActivateMode(FPersonaEditModes::SkeletonSelection);
-		}
-		else
-		{
-			Owner->DeactivateMode(FPersonaEditModes::SkeletonSelection);
-		}
-	}
+	return GetCurrentEditingCache()->GetMorphTargets();
+}
+
+float USkeletalMeshModelingToolsEditorMode::GetMorphTargetWeight(FName MorphTarget)
+{
+	return GetCurrentEditingCache()->GetMorphTargetWeight(MorphTarget);
+}
+
+FName USkeletalMeshModelingToolsEditorMode::HandleAddMorphTarget(FName InName)
+{
+	return GetCurrentEditingCache()->AddMorphTarget(InName);
+}
+
+FName USkeletalMeshModelingToolsEditorMode::HandleRenameMorphTarget(FName OldName, FName NewName)
+{
+	return GetCurrentEditingCache()->RenameMorphTarget(OldName, NewName);
+}
+
+void USkeletalMeshModelingToolsEditorMode::HandleRemoveMorphTargets(const TArray<FName>& Names)
+{
+	GetCurrentEditingCache()->RemoveMorphTargets(Names);
+}
+
+TArray<FName> USkeletalMeshModelingToolsEditorMode::HandleDuplicateMorphTargets(const TArray<FName>& Names)
+{
+	return GetCurrentEditingCache()->DuplicateMorphTargets(Names);
+}
+
+void USkeletalMeshModelingToolsEditorMode::HandleSetMorphTargetWeight(FName MorphTarget, float Weight)
+{
+	return GetCurrentEditingCache()->HandleSetMorphTargetWeight(MorphTarget, Weight);
+}
+
+bool USkeletalMeshModelingToolsEditorMode::GetMorphTargetAutoFill(FName MorphTarget)
+{
+	return GetCurrentEditingCache()->GetMorphTargetAutoFill(MorphTarget);
+}
+
+void USkeletalMeshModelingToolsEditorMode::HandleSetMorphTargetAutoFill(FName MorphTarget, bool bAutoFill, float PreviousOverrideWeight)
+{
+	return GetCurrentEditingCache()->HandleSetMorphTargetAutoFill(MorphTarget, bAutoFill, PreviousOverrideWeight);
 }
 
 
 void USkeletalMeshModelingToolsEditorMode::Exit()
 {
+	GetCurrentEditingCache()->ApplyChanges();
+	GetCurrentEditingCache()->Destroy();
+	CurrentEditingCache = nullptr;
+	
 	UEditorInteractiveToolsContext* InteractiveToolsContext = GetInteractiveToolsContext();
 	UE::TransformGizmoUtil::DeregisterTransformGizmoContextObject(InteractiveToolsContext);
 	UE::SkeletalMeshGizmoUtils::UnregisterTransformGizmoContextObject(InteractiveToolsContext);
@@ -355,6 +540,12 @@ void USkeletalMeshModelingToolsEditorMode::Exit()
 
 	// restore previous tool switching behavior
 	GetInteractiveToolsContext()->ToolManager->SetToolSwitchMode(ToolSwitchModeToRestoreOnExit);
+
+	if (Editor.IsValid())
+	{
+		Editor.Pin()->OnPreSaveAsset().RemoveAll(this);
+		Editor.Pin()->OnPreSaveAssetAs().RemoveAll(this);
+	}
 	
 #if ENABLE_STYLUS_SUPPORT
 	StylusStateTracker = nullptr;
@@ -363,6 +554,24 @@ void USkeletalMeshModelingToolsEditorMode::Exit()
 	UEdMode::Exit();
 }
 
+void USkeletalMeshModelingToolsEditorMode::Render(const FSceneView* View, FViewport* Viewport, FPrimitiveDrawInterface* PDI)
+{
+	Super::Render(View, Viewport, PDI);
+
+	if (GetCurrentEditingCache())
+	{
+		auto SetBoneDrawConfig = [Viewport](FSkelDebugDrawConfig& Config)
+			{
+				if (const FAnimationViewportClient* AnimViewportClient = static_cast<FAnimationViewportClient*>(Viewport->GetClient()))
+				{
+					Config.BoneDrawMode = AnimViewportClient->GetBoneDrawMode();
+					Config.BoneDrawSize = AnimViewportClient->GetBoneDrawSize();
+				}
+			};
+		GetCurrentEditingCache()->Render(PDI, SetBoneDrawConfig);
+	}
+}
+
 void USkeletalMeshModelingToolsEditorMode::RegisterExtensions()
 {
 	TArray<ISkeletalMeshModelingModeToolExtension*> Extensions = IModularFeatures::Get().GetModularFeatureImplementations<ISkeletalMeshModelingModeToolExtension>(
@@ -430,7 +639,9 @@ bool USkeletalMeshModelingToolsEditorMode::TryGetExtensionToolCommandGetter(
 
 void USkeletalMeshModelingToolsEditorMode::CreateToolkit()
 {
-	Toolkit = MakeShareable(new FSkeletalMeshModelingToolsEditorModeToolkit);
+	TSharedPtr<FSkeletalMeshModelingToolsEditorModeToolkit> NewToolKit = MakeShareable(new FSkeletalMeshModelingToolsEditorModeToolkit);
+	Toolkit = NewToolKit;
+	TypedToolkit = NewToolKit;
 }
 
 bool USkeletalMeshModelingToolsEditorMode::IsCompatibleWith(FEditorModeID OtherModeID) const
@@ -446,28 +657,48 @@ bool USkeletalMeshModelingToolsEditorMode::IsCompatibleWith(FEditorModeID OtherM
 void USkeletalMeshModelingToolsEditorMode::Tick(FEditorViewportClient* ViewportClient, float DeltaTime)
 {
 	Super::Tick(ViewportClient, DeltaTime);
+
+	if (GetCurrentEditingCache())
+	{
+		GetCurrentEditingCache()->Tick();
+	}
+
+	CacheChangeCountWatcher.CheckAndUpdate();
+	CacheDynamicMeshSkeletonStatusWatcher.CheckAndUpdate();
+	CacheSkeletonChangeCountWatcher.CheckAndUpdate();
 }
 
 bool USkeletalMeshModelingToolsEditorMode::HandleClick(FEditorViewportClient* InViewportClient, HHitProxy *HitProxy, const FViewportClick &Click)
 {
+	if (GetCurrentEditingCache()->HandleClick(HitProxy))
+	{
+		return true;
+	}
+	
 	// Making sure we can still select bones from viewport even if FPersonaEditModes::SkeletonSelection is deactivated
-	if (TSharedPtr<ISkeletalMeshEditorBinding> BindingPtr = GetEditorBinding())
+	if (TSharedPtr<ISkeletalMeshEditorBinding> EditorBindingPtr = GetEditorBinding())
 	{
-		TArray<FName> SelectedBones;
-		if (HitProxy && BindingPtr->GetNameFunction())
+		TArray<FName> Selected;
+		if (HitProxy && EditorBindingPtr->GetNameFunction())
 		{
-			if (TOptional<FName> BoneName = BindingPtr->GetNameFunction()(HitProxy))
+			if (TOptional<FName> BoneName = EditorBindingPtr->GetNameFunction()(HitProxy))
 			{
-				SelectedBones.Emplace(*BoneName);
+				Selected.Emplace(*BoneName);
 			}
 		}
-		
-		BindingPtr->GetNotifier().HandleNotification( SelectedBones, ESkeletalMeshNotifyType::BonesSelected);
+
+		GetModeBinding()->GetNotifier()->HandleNotification(Selected, ESkeletalMeshNotifyType::BonesSelected);
 	}
 	
 	return Super::HandleClick(InViewportClient, HitProxy, Click);
 }
 
+void USkeletalMeshModelingToolsEditorMode::PostUndo()
+{
+	Super::PostUndo();
+
+}
+
 bool USkeletalMeshModelingToolsEditorMode::ComputeBoundingBoxForViewportFocus(AActor* Actor, UPrimitiveComponent* PrimitiveComponent, FBox& InOutBox) const
 {
 	// if Tool supports custom Focus box, use that first
@@ -482,46 +713,51 @@ bool USkeletalMeshModelingToolsEditorMode::ComputeBoundingBoxForViewportFocus(AA
 		}
 	}
 
+	TArray<FName> Selection = GetCurrentEditingCache()->GetSelectedBones();
+	
+	
 	// focus using selected bones in skel mesh editor
-	if (const USkeletalMeshComponent* Component = Cast<USkeletalMeshComponent>(PrimitiveComponent))
+	if (!Selection.IsEmpty())
 	{
-		check(Component->GetSkeletalMeshAsset());
+		const FReferenceSkeleton& RefSkeleton = GetCurrentEditingCache()->GetEditingMeshComponent()->GetRefSkeleton();
+		TArray<FTransform> WorldTransforms = GetCurrentEditingCache()->GetComponentSpaceBoneTransforms();
+		FTransform MeshTransform = GetCurrentEditingCache()->GetTransform();
+		for (FTransform& Transform : WorldTransforms)
+		{
+			Transform *= MeshTransform;
+		}
 		
-		if (EditorBinding.IsValid())
+		TArray<FName> AllChildren;
+
+		for (const FName& BoneName: Selection)
 		{
-			const TArray<FName> Selection = EditorBinding.Pin()->GetSelectedBones();
-			if (!Selection.IsEmpty())
+			const int32 BoneIndex = RefSkeleton.FindBoneIndex(BoneName);
+			if (BoneIndex > INDEX_NONE && WorldTransforms.IsValidIndex(BoneIndex))
 			{
-				TArray<FName> AllChildren;
+				// enlarge box
+				InOutBox += WorldTransforms[BoneIndex].GetLocation();
 
-				const FReferenceSkeleton& RefSkeleton = Component->GetSkeletalMeshAsset()->GetRefSkeleton();
-				for (const FName& BoneName: Selection)
+				// get direct children
+				TArray<int32> Children;
+				RefSkeleton.GetDirectChildBones(BoneIndex, Children);
+				Algo::Transform(Children, AllChildren, [&RefSkeleton](int ChildrenIndex)
 				{
-					const int32 BoneIndex = RefSkeleton.FindBoneIndex(BoneName);
-					if (BoneIndex > INDEX_NONE)
-					{
-						// enlarge box
-						InOutBox += Component->GetBoneLocation(BoneName);
-
-						// get direct children
-						TArray<int32> Children;
-						RefSkeleton.GetDirectChildBones(BoneIndex, Children);
-						Algo::Transform(Children, AllChildren, [&RefSkeleton](int ChildrenIndex)
-						{
-							return RefSkeleton.GetBoneName(ChildrenIndex);
-						});
-					}
-				}
+					return RefSkeleton.GetBoneName(ChildrenIndex);
+				});
+			}
+		}
 
-				// enlarge box using direct children
-				for (const FName& BoneName: AllChildren)
-				{
-					InOutBox += Component->GetBoneLocation(BoneName);
-				}
-				
-				return true; 
+		// enlarge box using direct children
+		for (const FName& BoneName: AllChildren)
+		{
+			const int32 BoneIndex = RefSkeleton.FindBoneIndex(BoneName);
+			if (BoneIndex > INDEX_NONE && WorldTransforms.IsValidIndex(BoneIndex))
+			{
+				InOutBox += WorldTransforms[BoneIndex].GetLocation();	
 			}
 		}
+		
+		return true; 
 	}
 	
 	return Super::ComputeBoundingBoxForViewportFocus(Actor, PrimitiveComponent, InOutBox);
@@ -563,6 +799,249 @@ bool USkeletalMeshModelingToolsEditorMode::ShouldToolStartBeAllowed(const FStrin
 	return true;
 }
 
+bool USkeletalMeshModelingToolsEditorMode::StartTracking(FEditorViewportClient* InViewportClient, FViewport* InViewport)
+{
+	if (GetCurrentEditingCache()->IsDynamicMeshBoneManipulationEnabled())
+	{
+		if (GetCurrentEditingCache()->GetFirstSelectedBoneIndex() != INDEX_NONE)
+		{
+			const UE::Widget::EWidgetMode WidgetMode = InViewportClient->GetWidgetMode();
+
+			// we also allow undo/redo of bone manipulations
+			if (WidgetMode == UE::Widget::WM_Rotate)
+			{
+				GEditor->BeginTransaction(LOCTEXT("RotateBone", "Rotate Bone"));
+			}
+			else
+			{
+				GEditor->BeginTransaction(LOCTEXT("TranslateBone", "Translate Bone"));
+			}
+
+			GetCurrentEditingCache()->GetSkeletonPoser()->BeginPoseChange();
+			return true;
+		}
+	}
+	
+	return false;
+}
+
+bool USkeletalMeshModelingToolsEditorMode::EndTracking(FEditorViewportClient* InViewportClient, FViewport* InViewport)
+{
+	if (GetCurrentEditingCache()->GetSkeletonPoser()->IsRecordingPoseChange())
+	{
+		GetCurrentEditingCache()->GetSkeletonPoser()->EndPoseChange();
+		
+		GEditor->EndTransaction();
+		return true;
+	}
+	return false;
+}
+
+bool USkeletalMeshModelingToolsEditorMode::InputDelta(FEditorViewportClient* InViewportClient, FViewport* InViewport, FVector& InDrag,
+	FRotator& InRot, FVector& InScale)
+{
+	if (GetCurrentEditingCache()->GetSkeletonPoser()->IsRecordingPoseChange())
+	{
+		// Actually change bone transform
+		const EAxisList::Type CurrentAxis = InViewportClient->GetCurrentWidgetAxis();
+		const ECoordSystem CoordSystem = InViewportClient->GetWidgetCoordSystemSpace();
+		if (CurrentAxis != EAxisList::Type::None)
+		{
+			URefSkeletonPoser* Poser = GetCurrentEditingCache()->GetSkeletonPoser();
+
+			int32 BoneIndex = GetCurrentEditingCache()->GetFirstSelectedBoneIndex();
+			const FTransform& ComponentTransform = GetCurrentEditingCache()->GetTransform();
+			const FTransform& ComponentSpaceTransform = Poser->GetComponentSpaceTransform(BoneIndex);
+			FTransform BoneGlobal = ComponentSpaceTransform * ComponentTransform;
+
+			FTransform BaseTransform = BoneGlobal;
+			if (TOptional<FTransform> Additive = Poser->GetBoneAdditiveTransform(BoneIndex))
+			{
+				BaseTransform = BoneGlobal.GetRelativeTransformReverse(*Additive);
+			}
+
+			Poser->ModifyBoneAdditiveTransform(BoneIndex,
+				[&BaseTransform, &InDrag, &InRot, &InScale, &CoordSystem](FTransform& InTransform)
+				{
+					FVector Offset = BaseTransform.TransformVector(InDrag);
+
+					FVector RotAxis;
+					float RotAngle;
+					InRot.Quaternion().ToAxisAndAngle(RotAxis, RotAngle);
+
+					FVector4 BoneSpaceAxis = BaseTransform.TransformVectorNoScale(RotAxis);
+
+					//Calculate the new delta rotation
+					FQuat DeltaQuat(BoneSpaceAxis, RotAngle);
+					DeltaQuat.Normalize();
+
+					FQuat NewRotation = (InTransform * FTransform(DeltaQuat)).GetRotation();
+
+					FVector4 BoneSpaceScaleOffset;
+
+					if (CoordSystem == COORD_World)
+					{
+						BoneSpaceScaleOffset = BaseTransform.TransformVector(InScale);
+					}
+					else
+					{
+						BoneSpaceScaleOffset = InScale;
+					}
+
+					InTransform.SetTranslation(InTransform.GetTranslation() + Offset);
+					InTransform.SetRotation(NewRotation);
+					InTransform.SetScale3D(InTransform.GetScale3D() + BoneSpaceScaleOffset);
+				});
+
+			return true;
+		}	
+	}
+
+	return false;
+}
+
+void USkeletalMeshModelingToolsEditorMode::DrawHUD(FEditorViewportClient* ViewportClient, FViewport* Viewport, const FSceneView* View,
+	FCanvas* Canvas)
+{
+	if (GetCurrentEditingCache())
+	{
+		// Draw bone name
+		GetCurrentEditingCache()->DrawHUD(ViewportClient, Viewport, View, Canvas);
+	}
+	
+}
+
+bool USkeletalMeshModelingToolsEditorMode::AllowWidgetMove()
+{
+	if (GetCurrentEditingCache()->IsDynamicMeshSkeletonEnabled())
+	{
+		return ShouldDrawWidget();
+	}
+	return Super::AllowWidgetMove();
+}
+
+bool USkeletalMeshModelingToolsEditorMode::ShouldDrawWidget() const
+{
+	if (GetCurrentEditingCache()->IsDynamicMeshSkeletonEnabled())
+	{
+		if (GetCurrentEditingCache()->IsDynamicMeshBoneManipulationEnabled() && GetCurrentEditingCache()->GetFirstSelectedBoneIndex() != INDEX_NONE )
+		{
+			return true;
+		}
+
+		return false;
+	}
+
+	return Super::ShouldDrawWidget();
+}
+
+bool USkeletalMeshModelingToolsEditorMode::UsesTransformWidget() const
+{
+	if (GetCurrentEditingCache()->IsDynamicMeshSkeletonEnabled())
+	{
+		return true;
+	}
+	
+	return false;
+}
+
+bool USkeletalMeshModelingToolsEditorMode::UsesTransformWidget(UE::Widget::EWidgetMode CheckMode) const
+{
+	if (GetCurrentEditingCache()->IsDynamicMeshSkeletonEnabled())
+	{
+		return ShouldDrawWidget() && (CheckMode == UE::Widget::WM_Scale || CheckMode == UE::Widget::WM_Translate || CheckMode == UE::Widget::WM_Rotate);
+	}
+	
+	return false;
+}
+
+FVector USkeletalMeshModelingToolsEditorMode::GetWidgetLocation() const
+{
+	if (GetCurrentEditingCache()->IsDynamicMeshSkeletonEnabled())
+	{
+		int32 BoneIndex = GetCurrentEditingCache()->GetFirstSelectedBoneIndex();
+		URefSkeletonPoser* Poser = GetCurrentEditingCache()->GetSkeletonPoser();
+		FTransform ComponentTransform = GetCurrentEditingCache()->GetTransform();
+		if (BoneIndex != INDEX_NONE)
+		{
+			FTransform ComponentSpaceTransform = Poser->GetComponentSpaceTransform(BoneIndex);
+		
+			FTransform WorldSpaceTransform = ComponentSpaceTransform * ComponentTransform;
+
+			return WorldSpaceTransform.GetLocation();
+		}
+	
+		return ComponentTransform.GetLocation();
+	}
+
+	return Super::GetWidgetLocation();
+}
+
+
+
+bool USkeletalMeshModelingToolsEditorMode::GetCustomDrawingCoordinateSystem(FMatrix& InMatrix, void* InData)
+{
+	if (GetCurrentEditingCache()->IsDynamicMeshSkeletonEnabled())
+	{
+		const bool bIsParentMode = Owner ? Owner->GetCoordSystem() == COORD_Parent : false;
+	
+		const FReferenceSkeleton& RefSkeleton = GetCurrentEditingCache()->GetEditingMeshComponent()->GetRefSkeleton();
+		int32 BoneIndex = GetCurrentEditingCache()->GetFirstSelectedBoneIndex();
+		URefSkeletonPoser* Poser = GetCurrentEditingCache()->GetSkeletonPoser();
+		FTransform ComponentTransform = GetCurrentEditingCache()->GetTransform();
+		if (RefSkeleton.IsValidIndex(BoneIndex))
+		{
+			if (bIsParentMode)
+			{
+				const int32 ParentIndex = RefSkeleton.GetParentIndex(BoneIndex);
+				if (ParentIndex != INDEX_NONE)
+				{
+					BoneIndex = ParentIndex;
+				}
+			}
+		
+			FTransform BoneGlobalTransform = Poser->GetComponentSpaceTransform(BoneIndex) * ComponentTransform;	
+
+			InMatrix = BoneGlobalTransform.ToMatrixNoScale().RemoveTranslation();
+			return true;
+		}
+
+		return false;
+	}
+	
+	return Super::GetCustomDrawingCoordinateSystem(InMatrix, InData);
+}
+
+bool USkeletalMeshModelingToolsEditorMode::GetCustomInputCoordinateSystem(FMatrix& InMatrix, void* InData)
+{
+	if (GetCurrentEditingCache()->IsDynamicMeshSkeletonEnabled())
+	{
+		return GetCustomDrawingCoordinateSystem(InMatrix, InData);
+	}
+
+	return Super::GetCustomInputCoordinateSystem(InMatrix, InData);
+}
+
+USkeletalMeshBackedDynamicMeshComponent* USkeletalMeshModelingToolsEditorMode::GetComponent(UObject* SourceObject)
+{
+	return GetCurrentEditingCache()->GetEditingMeshComponent();
+}
+
+bool USkeletalMeshModelingToolsEditorMode::ShouldCommitToSkeletalMeshOnToolCommit()
+{
+	return ApplyMode == EApplyMode::ApplyOnToolExit;
+}
+
+TSharedPtr<ISkeletalMeshEditorBinding> USkeletalMeshModelingToolsEditorMode::GetModeBinding()
+{
+	if (!ModeBinding.IsValid())
+	{
+		ModeBinding = MakeShared<FSkeletalMeshModelingToolsEditorModeBinding>(this);
+	}
+
+	return ModeBinding;
+}
+
 void USkeletalMeshModelingToolsEditorMode::SetEditorBinding(const TWeakPtr<ISkeletalMeshEditor>& InSkeletalMeshEditor)
 {
 	if (!InSkeletalMeshEditor.IsValid())
@@ -572,7 +1051,13 @@ void USkeletalMeshModelingToolsEditorMode::SetEditorBinding(const TWeakPtr<ISkel
 	
 	EditorBinding = InSkeletalMeshEditor.Pin()->GetBinding();
 	Editor = InSkeletalMeshEditor;
+	Editor.Pin()->OnPreSaveAsset().AddUObject(this, &USkeletalMeshModelingToolsEditorMode::ApplyChanges);
+	Editor.Pin()->OnPreSaveAssetAs().AddUObject(this, &USkeletalMeshModelingToolsEditorMode::ApplyChanges);
+
 	SkeletalMesh = InSkeletalMeshEditor.Pin()->GetPersonaToolkit()->GetMesh();
+	SkeletalMesh->GetOnMeshChanged().AddUObject(this, &USkeletalMeshModelingToolsEditorMode::HandleSkeletalMeshChanged);
+	SkeletonReader = NewObject<USkeletonModifier>();
+	SkeletonReader->SetSkeletalMesh(SkeletalMesh.Get());
 	
 	if (USkeletalMeshEditorContextObject* ContextObject = UE::SkeletalMeshEditorUtils::GetEditorContextObject(GetInteractiveToolsContext()))
 	{
@@ -592,47 +1077,136 @@ TSharedPtr<ISkeletalMeshEditorBinding> USkeletalMeshModelingToolsEditorMode::Get
 	return EditorBinding.Pin();
 }
 
+TSharedPtr<ISkeletalMeshEditor> USkeletalMeshModelingToolsEditorMode::GetEditor()
+{
+	if (!Editor.IsValid())
+	{
+		return nullptr;
+	}
+
+	return Editor.Pin();
+}
+
 USkeletalMesh* USkeletalMeshModelingToolsEditorMode::GetSkeletalMesh() const
 {
 	return SkeletalMesh.Get();
 }
 
+void USkeletalMeshModelingToolsEditorMode::HandleSkeletalMeshChanged()
+{
+	if (GetCurrentEditingCache())
+	{
+		if (!GetCurrentEditingCache()->IsApplyingChanges())
+		{
+			int32 ChangeCount = GetCurrentEditingCache()->GetEditingMeshComponent()->GetChangeCount();
+			if (ChangeCount > 0)
+			{
+				FText Message = FText::Format(LOCTEXT("ExternalAssetChangeDiscardPendingToolChanges", "{0} change(s) discarded due to untracked external asset change"), ChangeCount);
+				GetToolManager()->DisplayMessage(Message, EToolMessageLevel::UserWarning);
+				ShowEditorMessage(ELogVerbosity::Type::Warning, Message);
+			}
+			
+			RecreateEditingCache(GetEditingLOD());
+		}
+	}
+	
+
+}
+
 bool USkeletalMeshModelingToolsEditorMode::CanSetEditingLOD()
 {
 	return !GetToolManager()->HasAnyActiveTool();
 }
 
 void USkeletalMeshModelingToolsEditorMode::SetEditingLOD(EMeshLODIdentifier InEditingLOD)
-{	
-	EditingLOD = InEditingLOD;
+{
+	if (GetCurrentEditingCache())
+	{
+		GetCurrentEditingCache()->ApplyChanges();
+	}
+	RecreateEditingCache(InEditingLOD);
+	SkeletonReader->ExternalUpdate(GetCurrentEditingCache()->GetEditingMeshComponent()->GetRefSkeleton());
 
 	TObjectPtr<UToolTargetManager> TargetManager = GetInteractiveToolsContext(EToolsContextScope::EdMode)->TargetManager;
 	if (USkeletalMeshComponentToolTargetFactory* SkeletalMeshTargetFactory = TargetManager->FindFirstFactoryByType<USkeletalMeshComponentToolTargetFactory>())
 	{
-		SkeletalMeshTargetFactory->SetActiveEditingLOD(EditingLOD);
+		SkeletalMeshTargetFactory->SetActiveEditingLOD(InEditingLOD);
 	}
 }
 
 EMeshLODIdentifier USkeletalMeshModelingToolsEditorMode::GetEditingLOD()
 {
-	return EditingLOD;
+	return GetCurrentEditingCache()->GetLOD();
 }
 
-void USkeletalMeshModelingToolsEditorMode::HideSkeletonForTool()
+void USkeletalMeshModelingToolsEditorMode::RecreateEditingCache(EMeshLODIdentifier InLOD)
 {
-	if (GetSkeletalMeshSkeletonDrawMode() != ESkeletonDrawMode::Hidden)
+	if (CurrentEditingCache)
+	{
+		CurrentEditingCache->Destroy();
+		EditingCacheNotifierBindScope.Reset();
+	}
+	
+	if (!Editor.IsValid())
 	{
-		SavedSkeletonDrawMode = GetSkeletalMeshSkeletonDrawMode();
-		SetSkeletalMeshSkeletonDrawMode(ESkeletonDrawMode::Hidden);
+		return;
 	}
+	
+	TSharedRef<IPersonaPreviewScene> PreviewScene = Editor.Pin()->GetPersonaToolkit()->GetPreviewScene();
+	UWorld* PreviewWorld = PreviewScene->GetWorld();
+	
+	CurrentEditingCache = NewObject<USkeletalMeshEditingCache>();
+
+	TWeakObjectPtr<USkeletalMeshModelingToolsEditorMode> WeakThis(this);
+	
+	USkeletalMeshEditingCache::FDelegates Delegates;
+	Delegates.ToggleSkeletalMeshBoneManipulationDelegate.BindUObject(this, &USkeletalMeshModelingToolsEditorMode::ToggleSkeletalMeshBoneManipulation);
+	Delegates.IsSkeletalMeshBoneManipulationEnabledDelegate.BindUObject(this, &USkeletalMeshModelingToolsEditorMode::IsSkeletalMeshBoneManipulationEnabled);
+	CurrentEditingCache->Spawn(PreviewWorld, GetSkelMeshComponent(), InLOD, Delegates);
+
+
+	using namespace UE::SkeletalMeshEditorUtils;
+	EditingCacheNotifierBindScope.Reset(new FSkeletalMeshNotifierBindScope(GetModeBinding()->GetNotifier(), CurrentEditingCache->GetNotifier()));
+}
+
+USkeletalMeshEditingCache* USkeletalMeshModelingToolsEditorMode::GetCurrentEditingCache() const
+{
+	return CurrentEditingCache;
+}
+
+bool USkeletalMeshModelingToolsEditorMode::HasUnappliedChanges() const
+{
+	return GetCurrentEditingCache()->GetEditingMeshComponent()->IsDirty();
+}
+
+void USkeletalMeshModelingToolsEditorMode::ApplyChanges()
+{
+	GetCurrentEditingCache()->ApplyChanges();
+}
+
+void USkeletalMeshModelingToolsEditorMode::DiscardChanges()
+{
+	GetCurrentEditingCache()->DiscardChanges();
+}
+
+USkeletalMeshModelingToolsEditorMode::EApplyMode USkeletalMeshModelingToolsEditorMode::GetApplyMode() const
+{
+	return ApplyMode;
+}
+
+void USkeletalMeshModelingToolsEditorMode::SetApplyMode(EApplyMode InApplyMode)
+{
+	ApplyMode = InApplyMode;
+}
+
+void USkeletalMeshModelingToolsEditorMode::HideSkeletonForTool()
+{
+	GetCurrentEditingCache()->HideSkeleton();
 }
 
 void USkeletalMeshModelingToolsEditorMode::ShowSkeletonForTool()
 {
-	if (GetSkeletalMeshSkeletonDrawMode() == ESkeletonDrawMode::Hidden)
-	{
-		SetSkeletalMeshSkeletonDrawMode(SavedSkeletonDrawMode);
-	}
+	GetCurrentEditingCache()->ShowSkeleton();
 }
 
 ISkeletalMeshEditingInterface* USkeletalMeshModelingToolsEditorMode::GetSkeletonInterface(UInteractiveTool* InTool)
@@ -644,4 +1218,30 @@ ISkeletalMeshEditingInterface* USkeletalMeshModelingToolsEditorMode::GetSkeleton
 	return static_cast<ISkeletalMeshEditingInterface*>(InTool->GetInterfaceAddress(USkeletalMeshEditingInterface::StaticClass()));
 }
 
+
+void USkeletalMeshModelingToolsEditorMode::ToggleSkeletalMeshBoneManipulation(bool bEnable)
+{
+	if (Owner)
+	{
+		if(bEnable)
+		{
+			Owner->ActivateMode(FPersonaEditModes::SkeletonSelection);
+		}
+		else
+		{
+			Owner->DeactivateMode(FPersonaEditModes::SkeletonSelection);
+		}
+	}
+}
+
+bool USkeletalMeshModelingToolsEditorMode::IsSkeletalMeshBoneManipulationEnabled()
+{
+	if (Owner)
+	{
+		return Owner->IsModeActive(FPersonaEditModes::SkeletonSelection);
+	}
+
+	return true;
+}
+
 #undef LOCTEXT_NAMESPACE
diff --git a/Engine/Plugins/Animation/SkeletalMeshModelingTools/Source/SkeletalMeshModelingTools/Private/SkeletalMeshModelingToolsEditorMode.h b/Engine/Plugins/Animation/SkeletalMeshModelingTools/Source/SkeletalMeshModelingTools/Private/SkeletalMeshModelingToolsEditorMode.h
--- a/Engine/Plugins/Animation/SkeletalMeshModelingTools/Source/SkeletalMeshModelingTools/Private/SkeletalMeshModelingToolsEditorMode.h
+++ b/Engine/Plugins/Animation/SkeletalMeshModelingTools/Source/SkeletalMeshModelingTools/Private/SkeletalMeshModelingToolsEditorMode.h
@@ -5,13 +5,23 @@
 #include "CoreMinimal.h"
 #include "EdMode.h"
 #include "SkeletalMeshModelingModeToolExtensions.h"
+#include "SkeletalMeshNotifier.h"
+#include "Interfaces/ISKMBackedDynaMeshComponentProvider.h"
 #include "Animation/DebugSkelMeshComponent.h"
+#include "Changes/ValueWatcher.h"
 #include "TargetInterfaces/MeshTargetInterfaceTypes.h"
 
 #include "SkeletalMeshModelingToolsEditorMode.generated.h"
 
-//class FStylusStateTracker;
 class FSkeletalMeshModelingToolsEditorModeToolkit;
+
+namespace UE::SkeletalMeshEditorUtils
+{
+	struct FSkeletalMeshNotifierBindScope;
+}
+
+class USkeletalMeshEditingCache;
+//class FStylusStateTracker;
 class UEdModeInteractiveToolsContext;
 class ISkeletalMeshNotifier;
 class ISkeletalMeshEditorBinding;
@@ -24,15 +34,44 @@ class FTabManager;
 enum class EToolSide;
 enum class EToolShutdownType : uint8;
 struct FToolBuilderState;
+class USkeletonModifier;
 
 namespace UE
 {
 	class IInteractiveToolCommandsInterface;
 }
 
+class USkeletalMeshModelingToolsEditorMode;
+class FSkeletalMeshModelingToolsEditorModeNotifier: public ISkeletalMeshNotifier
+{
+public:
+	FSkeletalMeshModelingToolsEditorModeNotifier(USkeletalMeshModelingToolsEditorMode* InEditorMode);
+	virtual void HandleNotification(const TArray<FName>& BoneNames, const ESkeletalMeshNotifyType InNotifyType) override;
+private:
+	TWeakObjectPtr<USkeletalMeshModelingToolsEditorMode> EditorMode;
+};
+
+class FSkeletalMeshModelingToolsEditorModeBinding: public ISkeletalMeshEditorBinding
+{
+public:
+	FSkeletalMeshModelingToolsEditorModeBinding(USkeletalMeshModelingToolsEditorMode* InEditorMode);
+
+	virtual TSharedPtr<ISkeletalMeshNotifier> GetNotifier() override;
+	virtual NameFunction GetNameFunction() override;
+	virtual TArray<FName> GetSelectedBones() const override;
+
+
+private:
+	TWeakObjectPtr<USkeletalMeshModelingToolsEditorMode> EditorMode;	
+	TSharedPtr<FSkeletalMeshModelingToolsEditorModeNotifier> Notifier;
+};
+
+
+
 UCLASS()
 class USkeletalMeshModelingToolsEditorMode : 
-	public UBaseLegacyWidgetEdMode
+	public UBaseLegacyWidgetEdMode,
+	public ISkeletalMeshBackedDynamicMeshComponentProvider 
 {
 	GENERATED_BODY()
 public:
@@ -46,46 +85,113 @@ public:
 	virtual void Initialize() override;
 	virtual void Enter() override;
 	virtual void Exit() override;
+	virtual void Render(const FSceneView* View, FViewport* Viewport, FPrimitiveDrawInterface* PDI) override;
 	virtual void CreateToolkit() override;
 	virtual bool IsCompatibleWith(FEditorModeID OtherModeID) const override;
-	virtual bool UsesTransformWidget() const override { return false; };
-	virtual bool UsesPropertyWidgets() const override { return false; };
+	virtual bool UsesPropertyWidgets() const override { return false; }
+
 	virtual void Tick(FEditorViewportClient* InViewportClient, float InDeltaTime) override;
-	virtual bool HandleClick(FEditorViewportClient* InViewportClient, HHitProxy *HitProxy, const FViewportClick &Click) override;
 	virtual bool ComputeBoundingBoxForViewportFocus(AActor* Actor, UPrimitiveComponent* PrimitiveComponent, FBox& InOutBox) const override;
 	virtual bool UsesToolkits() const override { return true; }
 	virtual bool ShouldToolStartBeAllowed(const FString& ToolIdentifier) const override;
+	virtual bool StartTracking(FEditorViewportClient* InViewportClient, FViewport* InViewport) override;
+	virtual bool EndTracking(FEditorViewportClient* InViewportClient, FViewport* InViewport) override;
+	virtual bool InputDelta(FEditorViewportClient* InViewportClient, FViewport* InViewport, FVector& InDrag, FRotator& InRot, FVector& InScale) override;
+	virtual void DrawHUD(FEditorViewportClient* ViewportClient, FViewport* Viewport, const FSceneView* View, FCanvas* Canvas) override;
+	virtual bool AllowWidgetMove() override;
+	virtual bool ShouldDrawWidget() const override;
+	virtual bool UsesTransformWidget() const override;
+	virtual bool UsesTransformWidget(UE::Widget::EWidgetMode CheckMode) const override;
+	virtual FVector GetWidgetLocation() const override;
+	virtual bool GetCustomDrawingCoordinateSystem(FMatrix& InMatrix, void* InData) override;
+	virtual bool GetCustomInputCoordinateSystem(FMatrix& InMatrix, void* InData) override;
+	virtual bool HandleClick(FEditorViewportClient* InViewportClient, HHitProxy *HitProxy, const FViewportClick &Click) override;
 
+	virtual void PostUndo() override;
+
+	// ISkeletalMeshBackedDynamicMeshComponentProvider
+	virtual USkeletalMeshBackedDynamicMeshComponent* GetComponent(UObject* SourceObject) override;
+	bool ShouldCommitToSkeletalMeshOnToolCommit() override;
+	
 	// binding
+	TSharedPtr<ISkeletalMeshEditorBinding> GetModeBinding();
+
 	void SetEditorBinding(const TWeakPtr<ISkeletalMeshEditor>& InSkeletalMeshEditor);
 	TSharedPtr<ISkeletalMeshEditorBinding> GetEditorBinding();
+	TSharedPtr<ISkeletalMeshEditor> GetEditor();
 	
 	USkeletalMesh* GetSkeletalMesh() const;
+	void HandleSkeletalMeshChanged();
 	
 	bool CanSetEditingLOD();
 	void SetEditingLOD(EMeshLODIdentifier EditingLOD);
 	EMeshLODIdentifier GetEditingLOD();
+
+	USkeletalMeshEditingCache* GetCurrentEditingCache() const;
+	bool HasUnappliedChanges() const;
+	void ApplyChanges();
+	void DiscardChanges();
+	
+	enum class EApplyMode: uint8
+	{
+		ApplyOnToolExit,
+		ApplyManually,
+	};
+	
+	EApplyMode GetApplyMode() const;
+	void SetApplyMode(EApplyMode InApplyMode);
+
 	
 	void HideSkeletonForTool();	
 	void ShowSkeletonForTool();
 	const TArray<FTransform>& GetComponentSpaceBoneTransforms() const;
 	void ToggleBoneManipulation(bool bEnable);
+	USkeletonModifier* GetSkeletonReader();
+	TArray<FName> GetSelectedBones() const;
+	
+	FName GetEditingMorphTarget();
+	void HandleSetEditingMorphTarget(FName InMorphTarget);
+
+	TArray<FName> GetMorphTargets();
 
+	float GetMorphTargetWeight(FName MorphTarget);
+	FName HandleAddMorphTarget(FName InName);
+	FName HandleRenameMorphTarget(FName OldName, FName NameName);
+	void HandleRemoveMorphTargets(const TArray<FName>& Names);
+	TArray<FName> HandleDuplicateMorphTargets(const TArray<FName>& Names);
+	
+	void HandleSetMorphTargetWeight(FName MorphTarget, float Weight);
+	bool GetMorphTargetAutoFill(FName MorphTarget);
+	void HandleSetMorphTargetAutoFill(FName MorphTarget, bool bAutoFill, float PreviousOverrideWeight);
 
 	TSharedPtr<FTabManager> GetAssociatedTabManager();
 	FSimpleMulticastDelegate& OnInitialized() { return OnInitializedDelegate; }
-	
+
+
 protected:
 	virtual void OnToolStarted(UInteractiveToolManager* Manager, UInteractiveTool* Tool) override;
 	void OnToolPostBuild(UInteractiveToolManager* InToolManager, EToolSide InSide, UInteractiveTool* InBuiltTool, UInteractiveToolBuilder* InToolBuilder, const FToolBuilderState& ToolState);
 	void OnToolEndedWithStatus(UInteractiveToolManager* Manager, UInteractiveTool* Tool, EToolShutdownType ShutdownType);
+
 private:
+	EApplyMode ApplyMode = EApplyMode::ApplyOnToolExit;
+	
+	TSharedPtr<FSkeletalMeshModelingToolsEditorModeBinding> ModeBinding;
 
 	TWeakPtr<ISkeletalMeshEditor> Editor;
 	TWeakPtr<ISkeletalMeshEditorBinding> EditorBinding;
 	TWeakObjectPtr<USkeletalMesh> SkeletalMesh;
+
+	TUniquePtr<UE::SkeletalMeshEditorUtils::FSkeletalMeshNotifierBindScope> EditorBindScope;
 	
-	EMeshLODIdentifier EditingLOD;
+	UPROPERTY()
+	TObjectPtr<USkeletalMeshEditingCache> CurrentEditingCache = nullptr;
+
+	void RecreateEditingCache(EMeshLODIdentifier InLOD);
+
+	TUniquePtr<UE::SkeletalMeshEditorUtils::FSkeletalMeshNotifierBindScope> EditingCacheNotifierBindScope;
+
+	FName EditingMorphTarget = NAME_None;
 	
 	// Stylus support is currently disabled; this is left in for reference if/when it is brought back
 	//TUniquePtr<FStylusStateTracker> StylusStateTracker;
@@ -96,16 +202,21 @@ private:
 	static ISkeletalMeshEditingInterface* GetSkeletonInterface(UInteractiveTool* InTool);
 
 	UDebugSkelMeshComponent* GetSkelMeshComponent() const;
+	void ToggleSkeletalMeshBoneManipulation(bool bEnable);
+	bool IsSkeletalMeshBoneManipulationEnabled();
 
-	bool bDeactivateOnPIEStartStateToRestore;
+	// A dummy read-only modifier to host the current edited skeleton for RefSkeletonTree UI
+	UPROPERTY()
+	TObjectPtr<USkeletonModifier> SkeletonReader = nullptr;
 
-	ESkeletonDrawMode GetSkeletalMeshSkeletonDrawMode() const;
-	void SetSkeletalMeshSkeletonDrawMode(ESkeletonDrawMode DrawMode) const;
-	ESkeletonDrawMode GetCurrentSkeletonDrawMode() const;
+	TValueWatcher<int32> CacheChangeCountWatcher;
+	TValueWatcher<bool> CacheDynamicMeshSkeletonStatusWatcher;
+	TValueWatcher<int32> CacheSkeletonChangeCountWatcher;
 	
-	ESkeletonDrawMode SavedSkeletonDrawMode = ESkeletonDrawMode::Default;
-
+	TWeakPtr<FSkeletalMeshModelingToolsEditorModeToolkit> TypedToolkit = nullptr;
 	
+	bool bDeactivateOnPIEStartStateToRestore;
+
 	void RegisterExtensions();
 	// Support extension tools having their own hotkey classes
 	TMap<FString, FExtensionToolDescription> ExtensionToolToInfo;
diff --git a/Engine/Plugins/Animation/SkeletalMeshModelingTools/Source/SkeletalMeshModelingTools/Private/SkeletalMeshModelingToolsEditorModeToolkit.cpp b/Engine/Plugins/Animation/SkeletalMeshModelingTools/Source/SkeletalMeshModelingTools/Private/SkeletalMeshModelingToolsEditorModeToolkit.cpp
--- a/Engine/Plugins/Animation/SkeletalMeshModelingTools/Source/SkeletalMeshModelingTools/Private/SkeletalMeshModelingToolsEditorModeToolkit.cpp
+++ b/Engine/Plugins/Animation/SkeletalMeshModelingTools/Source/SkeletalMeshModelingTools/Private/SkeletalMeshModelingToolsEditorModeToolkit.cpp
@@ -10,6 +10,7 @@
 #include "IDetailsView.h"
 #include "InteractiveTool.h"
 #include "InteractiveToolsContext.h"
+#include "ISkeletalMeshEditor.h"
 #include "ModelingToolsEditorModeStyle.h"
 #include "Modules/ModuleManager.h"
 #include "PropertyEditorModule.h"
@@ -20,14 +21,23 @@
 #include "Widgets/Text/STextBlock.h"
 #include "ModelingToolsManagerActions.h"
 #include "ModelingToolTargetEditorOnlyUtil.h"
+#include "PersonaTabs.h"
+#include "SEnumCombo.h"
+#include "SkeletalMeshEditingCache.h"
 #include "SkeletalMeshModelingModeToolExtensions.h"
 #include "SkeletalMeshModelingToolsEditorMode.h"
+#include "SMorphTargetManager.h"
 #include "BaseTools/MeshSurfacePointTool.h"
 #include "SPrimaryButton.h"
+#include "SkeletonModifier.h"
+#include "SReferenceSkeletonTree.h"
+#include "Components/SKMBackedDynaMeshComponent.h"
+#include "Editor/EditorWidgets/Public/SEnumCombo.h"
 #include "Toolkits/AssetEditorModeUILayer.h"
 #include "Widgets/Input/STextComboBox.h"
 #include "Widgets/Layout/SExpandableArea.h"
-
+#include "Widgets/Docking/SDockTab.h"
+#include "ScopedTransaction.h"
 
 
 #define LOCTEXT_NAMESPACE "SkeletalMeshModelingToolsEditorModeToolkit"
@@ -132,6 +142,8 @@ void FSkeletalMeshModelingToolsEditorModeToolkit::OnToolStarted(UInteractiveTool
 	ActiveToolIcon = FModelingToolsEditorModeStyle::Get()->GetOptionalBrush(ActiveToolIconName);
 	
 	GetToolkitHost()->AddViewportOverlayWidget(ViewportOverlayWidget.ToSharedRef());
+
+	AssetConfigPanel->SetExpanded_Animated(false);
 }
 
 void FSkeletalMeshModelingToolsEditorModeToolkit::OnToolEnded(UInteractiveToolManager* Manager, UInteractiveTool* Tool)
@@ -151,6 +163,8 @@ void FSkeletalMeshModelingToolsEditorModeToolkit::OnToolEnded(UInteractiveToolMa
 	ActiveToolName = FText::GetEmpty();
 	ClearNotification();
 	ClearWarning();
+	
+	AssetConfigPanel->SetExpanded_Animated(true);
 }
 
 FText FSkeletalMeshModelingToolsEditorModeToolkit::GetActiveToolDisplayName() const
@@ -164,6 +178,90 @@ FText FSkeletalMeshModelingToolsEditorModeToolkit::GetActiveToolMessage() const
 	return ActiveToolMessage;
 }
 
+void FSkeletalMeshModelingToolsEditorModeToolkit::InvokeUI()
+{
+	// Skipping parent class call here since persona already spawns a tool box tab
+	// invoking the UI here would create a second tool box tab
+	// FModeToolkit::InvokeUI();
+	
+	if (ModeUILayer.IsValid())
+	{
+		TSharedPtr<FAssetEditorModeUILayer> ModeUILayerPtr = ModeUILayer.Pin();
+		TSharedPtr<FTabManager> TabManager = ModeUILayerPtr->GetTabManager();
+		
+
+		SAssignNew(MorphTargetManagerWidget, SBorder)
+		.Padding(4.f)
+		.BorderImage(FAppStyle::GetBrush("ToolPanel.GroupBorder"));
+		
+		// Morph Target
+		{
+			const FName& MorphTargetTabID = FPersonaTabs::MorphTargetsID;
+			TSharedPtr<SDockTab> MorphTargetTab = TabManager->FindExistingLiveTab(MorphTargetTabID);
+			if (!MorphTargetTab.IsValid())
+			{
+				MorphTargetTab = TabManager->TryInvokeTab(MorphTargetTabID);
+			}
+		
+			{
+				DefaultMorphTargetWidget = MorphTargetTab->GetContent();
+			
+				MorphTargetTab->SetContent(MorphTargetManagerWidget.ToSharedRef());
+			}
+		}	
+	}
+
+	
+}
+
+void FSkeletalMeshModelingToolsEditorModeToolkit::ShutdownUI()
+{
+	//Skeleton
+	SkeletonNotifierBindScope.Reset();
+	
+	if (ModeUILayer.IsValid())
+	{
+		TSharedPtr<FAssetEditorModeUILayer> ModeUILayerPtr = ModeUILayer.Pin();
+		TSharedPtr<FTabManager> TabManager = ModeUILayerPtr->GetTabManager();
+
+		// Morph Target
+		{
+			const FName& MorphTargetTabID = FPersonaTabs::MorphTargetsID;
+			TSharedPtr<SDockTab> MorphTargetTab = TabManager->FindExistingLiveTab(MorphTargetTabID);
+			if (MorphTargetTab.IsValid())
+			{
+				// switch back to default widget if it was replaced before
+				if (DefaultMorphTargetWidget)
+				{
+					MorphTargetTab->SetContent(DefaultMorphTargetWidget.ToSharedRef());
+				}
+			}
+		}	
+	}
+
+	FModeToolkit::ShutdownUI();
+}
+
+TArray<FName> FSkeletalMeshModelingToolsEditorModeToolkit::GetSelectedBonesForDynamicMeshSkeleton()
+{
+	TArray<FName> SelectedBones;
+	RefSkeletonTree->GetSelectedBoneNames(SelectedBones);
+	return SelectedBones;
+}
+
+void FSkeletalMeshModelingToolsEditorModeToolkit::ResetDynamicMeshBoneTransforms(bool bSelectedOnly)
+{
+	FText Title = bSelectedOnly ? LOCTEXT("ResetBoneTransforms", "Reset Bone Transforms") : LOCTEXT("ResetAllBonesTransforms", "Reset All Bone Transforms");
+	
+	FScopedTransaction Transaction(Title);
+	
+	if (EditorMode.IsValid())
+	{
+		EditorMode->GetCurrentEditingCache()->ResetDynamicMeshBoneTransforms(bSelectedOnly);
+	}
+}
+
+
 void FSkeletalMeshModelingToolsEditorModeToolkit::RegisterPalettes()
 {
 	const FModelingToolsManagerCommands& Commands = FModelingToolsManagerCommands::Get();
@@ -417,36 +515,209 @@ void FSkeletalMeshModelingToolsEditorModeToolkit::OnEditorModeInitialized()
 
 	AssetLODMode->SetSelectedItem(AssetLODModes[0]);
 
+	int32 EnumIndex = ApplyModeEnums.IndexOfByKey(EditorMode->GetApplyMode());
+	ApplyModeComboBox->SetSelectedItem(ApplyModeStrings[EnumIndex]);
+
+	SAssignNew(RefSkeletonWidget, SBorder)
+	.Padding(4.f)
+	.BorderImage(FAppStyle::GetBrush("ToolPanel.GroupBorder"))
+	[
+		SAssignNew(RefSkeletonTree, SReferenceSkeletonTree)
+			.Modifier(EditorMode->GetSkeletonReader())
+			.bIsReadOnly(true)
+	];
+
+	SMorphTargetManager::FMorphTargetManagerDelegates Delegates;
+
+	Delegates.OnGetMorphTargets.BindUObject(EditorMode.Get(), &USkeletalMeshModelingToolsEditorMode::GetMorphTargets);
+	
+	Delegates.OnGetMorphTargetWeight.BindUObject(EditorMode.Get(), &USkeletalMeshModelingToolsEditorMode::GetMorphTargetWeight);
+	Delegates.OnSetMorphTargetWeight.BindUObject(EditorMode.Get(), &USkeletalMeshModelingToolsEditorMode::HandleSetMorphTargetWeight);
+	
+	Delegates.OnGetMorphTargetAutoFill.BindUObject(EditorMode.Get(), &USkeletalMeshModelingToolsEditorMode::GetMorphTargetAutoFill);
+	Delegates.OnSetMorphTargetAutoFill.BindUObject(EditorMode.Get(), &USkeletalMeshModelingToolsEditorMode::HandleSetMorphTargetAutoFill);
+	
+	Delegates.OnGetEditingMorphTarget.BindUObject(EditorMode.Get(), &USkeletalMeshModelingToolsEditorMode::GetEditingMorphTarget);
+	Delegates.OnSetEditingMorphTarget.BindUObject(EditorMode.Get(), &USkeletalMeshModelingToolsEditorMode::HandleSetEditingMorphTarget);
+	
+	Delegates.OnAddNewMorphTarget.BindUObject(EditorMode.Get(), &USkeletalMeshModelingToolsEditorMode::HandleAddMorphTarget);
+	Delegates.OnRenameMorphTarget.BindUObject(EditorMode.Get(), &USkeletalMeshModelingToolsEditorMode::HandleRenameMorphTarget);
+	Delegates.OnRemoveMorphTargets.BindUObject(EditorMode.Get(), &USkeletalMeshModelingToolsEditorMode::HandleRemoveMorphTargets);
+	Delegates.OnDuplicateMorphTargets.BindUObject(EditorMode.Get(), &USkeletalMeshModelingToolsEditorMode::HandleDuplicateMorphTargets);
+	
+	SAssignNew(MorphTargetManager, SMorphTargetManager)
+			.Delegates(Delegates);
+	MorphTargetManagerWidget->SetContent(MorphTargetManager.ToSharedRef());
+
+	// Make sure viewport selection is routed to the default skeleton tree by default
+	SkeletonNotifierBindScope.Reset(
+		new UE::SkeletalMeshEditorUtils::FSkeletalMeshNotifierBindScope(
+			EditorMode->GetModeBinding()->GetNotifier(),
+			EditorMode->GetEditorBinding()->GetNotifier()));
+
+
+	
 }
-TSharedPtr<SWidget> FSkeletalMeshModelingToolsEditorModeToolkit::MakeFooterWidget()
+
+void FSkeletalMeshModelingToolsEditorModeToolkit::ShowDynamicMeshSkeletonTree()
 {
-	AssetLODMode = SNew(STextComboBox)
-		.OptionsSource(&AssetLODModes)
-		.IsEnabled_Lambda([this]()
+	if (ModeUILayer.IsValid())
+	{
+		TSharedPtr<FAssetEditorModeUILayer> ModeUILayerPtr = ModeUILayer.Pin();
+		TSharedPtr<FTabManager> TabManager = ModeUILayerPtr->GetTabManager();
+
+		// Skeleton Tree
 		{
-			if (!EditorMode.IsValid())
+			const FName& SkeletonTreeId = FPersonaTabs::SkeletonTreeViewID;
+			TSharedPtr<SDockTab> SkeletonTab = TabManager->FindExistingLiveTab(SkeletonTreeId);
+			if (!SkeletonTab.IsValid())
+			{
+				SkeletonTab = TabManager->TryInvokeTab(SkeletonTreeId);
+			}
+		
+			if (SkeletonTab->GetContent() != RefSkeletonWidget)
 			{
-				return false;
+				DefaultSkeletonWidget = SkeletonTab->GetContent();
+			
+				// switch between SSkeletonTree and SReferenceSkeletonTree
+				SkeletonTab->SetContent(RefSkeletonWidget.ToSharedRef());
+
+				
+				// Make sure viewport selection is routed to the ref skeleton tree
+				SkeletonNotifierBindScope.Reset(
+					new UE::SkeletalMeshEditorUtils::FSkeletalMeshNotifierBindScope(
+						EditorMode->GetModeBinding()->GetNotifier(),
+						RefSkeletonTree->GetNotifier()));
+				
+				TArray<FName> SelectedBones = EditorMode->GetSelectedBones();
+				
+				EditorMode->GetEditorBinding()->GetNotifier()->HandleNotification({}, ESkeletalMeshNotifyType::BonesSelected);
+				RefSkeletonTree->GetNotifier()->HandleNotification({}, ESkeletalMeshNotifyType::HierarchyChanged);
+				RefSkeletonTree->GetNotifier()->HandleNotification(SelectedBones, ESkeletalMeshNotifyType::BonesSelected);
+
+
+				const TSharedRef<FUICommandList>& CommandList = EditorMode->GetEditor()->GetToolkitCommands();
+				DefaultResetBoneTransformsAction = *(CommandList->GetActionForCommand(EditorMode->GetEditor()->GetResetBoneTransformsCommand()));
+				DefaultResetAllBoneTransformsAction = *(CommandList->GetActionForCommand(EditorMode->GetEditor()->GetResetAllBonesTransformsCommand()));
+
+				// Reset transforms
+				static constexpr bool bSelectedOnly = true;
+				CommandList->MapAction(EditorMode->GetEditor()->GetResetBoneTransformsCommand(),
+					FExecuteAction::CreateSP(this, &FSkeletalMeshModelingToolsEditorModeToolkit::ResetDynamicMeshBoneTransforms ,bSelectedOnly));
+	
+				CommandList->MapAction(EditorMode->GetEditor()->GetResetAllBonesTransformsCommand(),
+					FExecuteAction::CreateSP(this, &FSkeletalMeshModelingToolsEditorModeToolkit::ResetDynamicMeshBoneTransforms ,!bSelectedOnly));
 			}
+		}
+	}
+}
 
-			return EditorMode->CanSetEditingLOD();
-		})
-		.OnSelectionChanged_Lambda([this](TSharedPtr<FString> String, ESelectInfo::Type)
+void FSkeletalMeshModelingToolsEditorModeToolkit::ShowSkeletalMeshSkeletonTree()
+{
+	if (ModeUILayer.IsValid())
+	{
+		TSharedPtr<FAssetEditorModeUILayer> ModeUILayerPtr = ModeUILayer.Pin();
+		TSharedPtr<FTabManager> TabManager = ModeUILayerPtr->GetTabManager();
+
+		// Skeleton Tree
 		{
-			int32 Index = AssetLODModes.IndexOfByKey(String);
-			check(AssetAvailableLODs.Num() == AssetLODModes.Num());
-			EMeshLODIdentifier NewSelectedLOD = AssetAvailableLODs[Index];
+			const FName& SkeletonTreeViewID = FPersonaTabs::SkeletonTreeViewID;
+			TSharedPtr<SDockTab> SkeletonTab = TabManager->FindExistingLiveTab(SkeletonTreeViewID);
+			if (!SkeletonTab.IsValid())
+			{
+				SkeletonTab = TabManager->TryInvokeTab(SkeletonTreeViewID);
+			}
 
-			if (EditorMode.IsValid())
+			// switch back to default widget if it was replaced before
+			if (DefaultSkeletonWidget)
 			{
-				EditorMode->SetEditingLOD(NewSelectedLOD);
+				SkeletonTab->SetContent(DefaultSkeletonWidget.ToSharedRef());
 			}
-		});
+		}
+
+		// Make sure viewport selection is routed to the default skeleton tree
+		SkeletonNotifierBindScope.Reset(
+			new UE::SkeletalMeshEditorUtils::FSkeletalMeshNotifierBindScope(
+				EditorMode->GetModeBinding()->GetNotifier(),
+				EditorMode->GetEditorBinding()->GetNotifier()));
+
+		TArray<FName> SelectedBones = EditorMode->GetSelectedBones();
+		RefSkeletonTree->GetNotifier()->HandleNotification({}, ESkeletalMeshNotifyType::BonesSelected);
+		EditorMode->GetEditorBinding()->GetNotifier()->HandleNotification({}, ESkeletalMeshNotifyType::HierarchyChanged);
+		EditorMode->GetEditorBinding()->GetNotifier()->HandleNotification(SelectedBones, ESkeletalMeshNotifyType::BonesSelected);
+
+		// Rebind commands it was replaced before
+		if (DefaultSkeletonWidget)
+		{
+			const TSharedRef<FUICommandList>& CommandList = EditorMode->GetEditor()->GetToolkitCommands();
+
+			// Reset transforms
+			CommandList->MapAction(EditorMode->GetEditor()->GetResetBoneTransformsCommand(), DefaultResetBoneTransformsAction);
+			CommandList->MapAction(EditorMode->GetEditor()->GetResetAllBonesTransformsCommand(), DefaultResetAllBoneTransformsAction);
+		}
+	}	
+}
 
-	AssetLODModeLabel = SNew(STextBlock)
-	.Text(LOCTEXT("ActiveLODLabel", "Editing LOD"))
-	.ToolTipText(LOCTEXT("ActiveLODLabelToolTip", "Select the LOD to be used when editing an existing mesh."));
+void FSkeletalMeshModelingToolsEditorModeToolkit::RefreshMorphTargetManager()
+{
+	MorphTargetManager->RefreshList();
+}
 
+TSharedPtr<SWidget> FSkeletalMeshModelingToolsEditorModeToolkit::MakeFooterWidget()
+{
+	// LOD picker widget
+	{
+		AssetLODMode = SNew(STextComboBox)
+			.OptionsSource(&AssetLODModes)
+			.IsEnabled(this, &FSkeletalMeshModelingToolsEditorModeToolkit::CanChangeAssetEditingSettings)
+			.OnSelectionChanged_Lambda([this](TSharedPtr<FString> String, ESelectInfo::Type)
+			{
+				int32 Index = AssetLODModes.IndexOfByKey(String);
+				check(AssetAvailableLODs.Num() == AssetLODModes.Num());
+				EMeshLODIdentifier NewSelectedLOD = AssetAvailableLODs[Index];
+
+				if (EditorMode.IsValid())
+				{
+					EditorMode->SetEditingLOD(NewSelectedLOD);
+				}
+			});
+
+		AssetLODModeLabel = SNew(STextBlock)
+		.Text(LOCTEXT("ActiveLODLabel", "Editing LOD"))
+		.ToolTipText(LOCTEXT("ActiveLODLabelToolTip", "Select the LOD to be used when editing an existing mesh."));
+	}
+
+	// Apply Mode Picker Widget
+	{
+		ApplyModeLabel =
+			SNew(STextBlock)
+			.Text(LOCTEXT("ApplyModeLabel", "Apply Mode"))
+			.ToolTipText(LOCTEXT("ApplyModeToolTip", "Select when tool changes should be applied to the asset"));
+
+		ApplyModeEnums.Add(USkeletalMeshModelingToolsEditorMode::EApplyMode::ApplyOnToolExit);
+		ApplyModeStrings.Add(MakeShared<FString>(TEXT("Apply On Tool Exit")));
+		
+		ApplyModeEnums.Add(USkeletalMeshModelingToolsEditorMode::EApplyMode::ApplyManually);
+		ApplyModeStrings.Add(MakeShared<FString>(TEXT("Apply Manually")));
+			
+		ApplyModeComboBox =
+			SNew(STextComboBox)
+			.OptionsSource(&ApplyModeStrings)
+			.IsEnabled(this, &FSkeletalMeshModelingToolsEditorModeToolkit::CanChangeAssetEditingSettings)
+			.OnSelectionChanged_Lambda([this](TSharedPtr<FString> String, ESelectInfo::Type)
+			{
+				int32 Index = ApplyModeStrings.IndexOfByKey(String);
+				USkeletalMeshModelingToolsEditorMode::EApplyMode ApplyMode = ApplyModeEnums[Index];
+
+				if (EditorMode.IsValid())
+				{
+					EditorMode->SetApplyMode(ApplyMode);
+				}
+			});
+	}
+	
+
+	
 	const TSharedPtr<SVerticalBox> Content = SNew(SVerticalBox)
 	+ SVerticalBox::Slot().HAlign(HAlign_Fill)
 	.Padding(0, 8, 0, 4)
@@ -462,9 +733,70 @@ TSharedPtr<SWidget> FSkeletalMeshModelingToolsEditorModeToolkit::MakeFooterWidge
 				+ SHorizontalBox::Slot().Padding(0).HAlign(HAlign_Fill).Padding(0).FillWidth(4.f)
 					[ AssetLODMode->AsShared() ]
 			]
+			+SVerticalBox::Slot()
+			.AutoHeight()
+			[
+				SNew(SHorizontalBox)
+				+ SHorizontalBox::Slot().Padding(0).HAlign(HAlign_Left).VAlign(VAlign_Center).FillWidth(2.f)
+					[ ApplyModeLabel->AsShared() ]
+				+ SHorizontalBox::Slot().Padding(0).HAlign(HAlign_Fill).Padding(0).FillWidth(4.f)
+					[ ApplyModeComboBox->AsShared() ]
+			]
+			
+			+SVerticalBox::Slot()
+			.AutoHeight()
+			[
+				SNew(SHorizontalBox)
+				+SHorizontalBox::Slot()
+				.AutoWidth()
+				[
+					SNew(SBox)
+					.HeightOverride(50.f)
+					.HAlign(HAlign_Fill)
+					.Padding(0, 10, 5,10)
+					[
+						SNew(SButton)
+						.ButtonStyle(FAppStyle::Get(), "FlatButton.Default")
+						.ToolTipText(LOCTEXT("DiscardButtonToolTip", "Discard Pending Changes"))
+						.ForegroundColor(FLinearColor::White)
+						.HAlign(HAlign_Center)
+						.VAlign(VAlign_Center)
+						.IsEnabled(this, &FSkeletalMeshModelingToolsEditorModeToolkit::IsDiscardButtonEnabled)
+						.OnClicked(this, &FSkeletalMeshModelingToolsEditorModeToolkit::OnDiscardButtonPressed)
+						[
+							SNew( SImage )
+							.Image(FAppStyle::GetBrush("Icons.Delete"))
+							.ColorAndOpacity( FSlateColor::UseForeground() )
+						]		
+					]	
+				]
+				+SHorizontalBox::Slot()
+				.FillWidth(1.f)
+				[
+					SNew(SBox)
+					.HeightOverride(50.f)
+					.HAlign(HAlign_Fill)
+					.Padding(5, 10, 0, 10)
+					[
+						SNew(SButton)
+						.ButtonStyle(FAppStyle::Get(), FName("FlatButton.Success"))
+						.ForegroundColor(FLinearColor::White)
+						.HAlign(HAlign_Center)
+						.VAlign(VAlign_Center)
+						.IsEnabled(this, &FSkeletalMeshModelingToolsEditorModeToolkit::IsApplyButtonEnabled)
+						.OnClicked(this, &FSkeletalMeshModelingToolsEditorModeToolkit::OnApplyButtonPressed)
+						[
+							SNew(STextBlock)
+							.Text(this, &FSkeletalMeshModelingToolsEditorModeToolkit::GetApplyButtonText)
+							.Font(FAppStyle::GetFontStyle(TEXT("PropertyWindow.NormalFont")))
+						]		
+					]	
+				]	
+			
+			]
 		];
 
-	TSharedPtr<SExpandableArea> AssetConfigPanel = SNew(SExpandableArea)
+		AssetConfigPanel = SNew(SExpandableArea)
 		.HeaderPadding(FMargin(0.f))
 		.Padding(FMargin(8.f))
 		.BorderImage(FAppStyle::Get().GetBrush("DetailsView.CategoryTop"))
@@ -486,6 +818,47 @@ TSharedPtr<SWidget> FSkeletalMeshModelingToolsEditorModeToolkit::MakeFooterWidge
 	return AssetConfigPanel;
 }
 
+bool FSkeletalMeshModelingToolsEditorModeToolkit::IsApplyButtonEnabled() const
+{
+	return CanChangeAssetEditingSettings() && EditorMode->HasUnappliedChanges();
+}
+
+FReply FSkeletalMeshModelingToolsEditorModeToolkit::OnApplyButtonPressed()
+{
+	FScopedTransaction Transaction(LOCTEXT("ApplyChangesToAsset","Apply Changes To Skeletal Mesh"));
+	EditorMode->ApplyChanges();
+	
+	return FReply::Handled();
+}
+
+FText FSkeletalMeshModelingToolsEditorModeToolkit::GetApplyButtonText() const
+{
+	int32 ChangeCount = EditorMode->GetCurrentEditingCache()->GetEditingMeshComponent()->GetChangeCount();
+	if (ChangeCount == 0)
+	{
+		return FText::Format(LOCTEXT("ApplyButtonZeroChangeText", "No Pending Changes"), FText::AsNumber(ChangeCount));
+	}
+	
+	return FText::Format(LOCTEXT("ApplyButtonChangeText", "Apply {0}|plural(one=Change, other=Changes)"), FText::AsNumber(ChangeCount));
+}
+
+bool FSkeletalMeshModelingToolsEditorModeToolkit::IsDiscardButtonEnabled() const
+{
+	return IsApplyButtonEnabled();
+}
+
+FReply FSkeletalMeshModelingToolsEditorModeToolkit::OnDiscardButtonPressed()
+{
+	FScopedTransaction Transaction(LOCTEXT("DiscardChanges","Discard Changes"));
+	EditorMode->DiscardChanges();
+	return FReply::Handled();
+}
+
+bool FSkeletalMeshModelingToolsEditorModeToolkit::CanChangeAssetEditingSettings() const
+{
+	return !GetScriptableEditorMode()->GetToolManager(EToolsContextScope::EdMode)->HasAnyActiveTool();
+}
+
 void FSkeletalMeshModelingToolsEditorModeToolkit::UpdateActiveToolProperties(UInteractiveTool* Tool)
 {
 	UInteractiveTool* CurTool = GetScriptableEditorMode()->GetToolManager(EToolsContextScope::EdMode)->GetActiveTool(EToolSide::Left);
diff --git a/Engine/Plugins/Animation/SkeletalMeshModelingTools/Source/SkeletalMeshModelingTools/Private/SkeletalMeshModelingToolsEditorModeToolkit.h b/Engine/Plugins/Animation/SkeletalMeshModelingTools/Source/SkeletalMeshModelingTools/Private/SkeletalMeshModelingToolsEditorModeToolkit.h
--- a/Engine/Plugins/Animation/SkeletalMeshModelingTools/Source/SkeletalMeshModelingTools/Private/SkeletalMeshModelingToolsEditorModeToolkit.h
+++ b/Engine/Plugins/Animation/SkeletalMeshModelingTools/Source/SkeletalMeshModelingTools/Private/SkeletalMeshModelingToolsEditorModeToolkit.h
@@ -4,10 +4,12 @@
 
 #include "CoreMinimal.h"
 #include "SkeletalMeshEditorUtils.h"
+#include "SkeletalMeshModelingToolsEditorMode.h"
 #include "StatusBarSubsystem.h"
 #include "Toolkits/BaseToolkit.h"
 
 
+class SMorphTargetManager;
 class USkeletalMeshModelingToolsEditorMode;
 enum class EMeshLODIdentifier;
 class IToolkitHost;
@@ -31,7 +33,15 @@ public:
 	virtual void OnToolEnded(UInteractiveToolManager* Manager, UInteractiveTool* Tool) override;
 	virtual FText GetActiveToolDisplayName() const override;
 	virtual FText GetActiveToolMessage() const override;
+	virtual void InvokeUI() override;
+	virtual void ShutdownUI() override;
 
+	TArray<FName> GetSelectedBonesForDynamicMeshSkeleton();
+	void ResetDynamicMeshBoneTransforms(bool bSelectedOnly);
+	void ShowDynamicMeshSkeletonTree();
+	void ShowSkeletalMeshSkeletonTree();
+
+	void RefreshMorphTargetManager();
 private:
 	void PostNotification(const FText& Message);
 	void ClearNotification();
@@ -41,7 +51,16 @@ private:
 
 
 	void OnEditorModeInitialized();
+
+
 	TSharedPtr<SWidget> MakeFooterWidget();
+	bool IsApplyButtonEnabled() const;
+	FReply OnApplyButtonPressed();
+	FText GetApplyButtonText() const;
+	bool IsDiscardButtonEnabled() const;
+	FReply OnDiscardButtonPressed();
+	bool CanChangeAssetEditingSettings() const;
+	
 	void UpdateActiveToolProperties(UInteractiveTool* Tool);
 
 	void RegisterPalettes();
@@ -64,10 +83,30 @@ private:
 	TSharedPtr<STextBlock> ModeHeaderArea;
 	TSharedPtr<STextBlock> ToolWarningArea;
 
+	TSharedPtr<SExpandableArea> AssetConfigPanel;
+	
 	TArray<EMeshLODIdentifier> AssetAvailableLODs;
 	TArray<TSharedPtr<FString>> AssetLODModes;
 	TSharedPtr<STextBlock> AssetLODModeLabel;
 	TSharedPtr<STextComboBox> AssetLODMode;
+
+	TArray<USkeletalMeshModelingToolsEditorMode::EApplyMode> ApplyModeEnums;
+	TArray<TSharedPtr<FString>> ApplyModeStrings;
+	TSharedPtr<STextBlock> ApplyModeLabel;
+	TSharedPtr<STextComboBox> ApplyModeComboBox;
 	
+	TSharedPtr<SWidget> DefaultSkeletonWidget;
+	TSharedPtr<SBorder> RefSkeletonWidget;
+	TSharedPtr<SReferenceSkeletonTree> RefSkeletonTree;
+
+	TSharedPtr<SWidget> DefaultMorphTargetWidget; 
+	TSharedPtr<SBorder> MorphTargetManagerWidget; 
+	TSharedPtr<SMorphTargetManager> MorphTargetManager; 
+
+	TUniquePtr<UE::SkeletalMeshEditorUtils::FSkeletalMeshNotifierBindScope> SkeletonNotifierBindScope;
+
 	TWeakObjectPtr<USkeletalMeshModelingToolsEditorMode> EditorMode;
+
+	FUIAction DefaultResetBoneTransformsAction;
+	FUIAction DefaultResetAllBoneTransformsAction;
 };
diff --git a/Engine/Plugins/Animation/SkeletalMeshModelingTools/Source/SkeletalMeshModelingTools/Private/ToolTargets/SKMBackedDynaMeshComponentToolTarget.cpp b/Engine/Plugins/Animation/SkeletalMeshModelingTools/Source/SkeletalMeshModelingTools/Private/ToolTargets/SKMBackedDynaMeshComponentToolTarget.cpp
new file mode 100644
--- /dev/null
+++ b/Engine/Plugins/Animation/SkeletalMeshModelingTools/Source/SkeletalMeshModelingTools/Private/ToolTargets/SKMBackedDynaMeshComponentToolTarget.cpp
@@ -0,0 +1,139 @@
+﻿// Copyright Epic Games, Inc. All Rights Reserved.
+
+#include "ToolTargets/SKMBackedDynaMeshComponentToolTarget.h"
+#include "Components/SkeletalMeshComponent.h"
+#include "Engine/SkeletalMesh.h"
+
+#include "Components/SKMBackedDynaMeshComponent.h"
+
+#include UE_INLINE_GENERATED_CPP_BY_NAME(SKMBackedDynaMeshComponentToolTarget)
+
+#define LOCTEXT_NAMESPACE "SKMBackedDynaMeshComponentToolTarget"
+
+
+FReferenceSkeleton USkeletalMeshBackedDynamicMeshComponentToolTarget::GetSkeleton() const
+{
+	return GetComponent()->GetRefSkeleton();
+	
+}
+
+void USkeletalMeshBackedDynamicMeshComponentToolTarget::CommitSkeletonModifier(USkeletonModifier* InModifier)
+{
+	Super::CommitSkeletonModifier(InModifier);
+
+	GetComponent()->HandleSkeletonChange(InModifier);
+
+	if (bCommitToSkeletalMeshOnToolCommit)
+	{
+		GetComponent()->CommitToSkeletalMesh();
+	}
+}
+
+void USkeletalMeshBackedDynamicMeshComponentToolTarget::CommitDynamicMeshChange(TUniquePtr<FToolCommandChange> Change, const FText& ChangeMessage)
+{
+	Super::CommitDynamicMeshChange(MoveTemp(Change), ChangeMessage);
+	GetComponent()->MarkDirty();
+
+	if (bCommitToSkeletalMeshOnToolCommit)
+	{
+		GetComponent()->CommitToSkeletalMesh();
+	}
+}
+
+USkeletalMesh* USkeletalMeshBackedDynamicMeshComponentToolTarget::GetSkeletalMesh() const
+{
+	return GetComponent()->GetSkeletalMesh();
+}
+
+void USkeletalMeshBackedDynamicMeshComponentToolTarget::SetOwnerVisibility(bool bVisible) const
+{
+	GetComponent()->ForwardVisibilityChangeRequest(bVisible);
+}
+
+USkeletalMeshBackedDynamicMeshComponent* USkeletalMeshBackedDynamicMeshComponentToolTarget::GetComponent() const
+{
+	return CastChecked<USkeletalMeshBackedDynamicMeshComponent>(Component);
+}
+
+void USkeletalMeshBackedDynamicMeshComponentToolTargetFactory::Init(ISkeletalMeshBackedDynamicMeshComponentProvider* InComponentProvider)
+{
+	ComponentProvider = InComponentProvider;
+}
+
+bool USkeletalMeshBackedDynamicMeshComponentToolTargetFactory::CanBuildTarget(
+	UObject* SourceObject,
+	const FToolTargetTypeRequirements& Requirements
+	) const
+{
+	// We are using an exact cast here to prevent subclasses, which might not meet all
+	// requirements for functionality such as the deprecated DestructibleMeshComponent, from 
+	// being caught up as valid targets.
+	// If you want to make the tool target work with some subclass of USkeletalMeshComponent,
+	// just add another factory that allows that class specifically(but make sure that
+	// GetMeshDescription and such work properly)
+	
+	bool bValid = Cast<USkeletalMeshComponent>(SourceObject) && Cast<USkeletalMeshComponent>(SourceObject)->GetSkeletalMeshAsset() &&
+		ExactCast<USkeletalMesh>(Cast<USkeletalMeshComponent>(SourceObject)->GetSkeletalMeshAsset()) &&
+		!ExactCast<USkeletalMesh>(Cast<USkeletalMeshComponent>(SourceObject)->GetSkeletalMeshAsset())->GetOutermost()->bIsCookedForEditor;
+	if (!bValid)
+	{
+		return false;
+	}
+
+	if (!CanWriteToSource(SourceObject))
+	{
+		return false;
+	}
+	
+	if (!ComponentProvider.IsValid())
+	{
+		return false;
+	}
+
+	USkeletalMeshBackedDynamicMeshComponent* Component = ComponentProvider->GetComponent(SourceObject);
+
+	if (!Component)
+	{
+		return false;
+	}
+
+	return Super::CanBuildTarget(Component, Requirements);
+}
+
+UToolTarget* USkeletalMeshBackedDynamicMeshComponentToolTargetFactory::BuildTarget(UObject* SourceObject, const FToolTargetTypeRequirements& Requirements)
+{
+	if (!ComponentProvider.IsValid())
+	{
+		return nullptr;
+	}
+
+	USkeletalMeshBackedDynamicMeshComponent* Component = ComponentProvider->GetComponent(SourceObject);
+
+	if (!Component)
+	{
+		return nullptr;
+	}
+
+	USkeletalMeshBackedDynamicMeshComponentToolTarget* Target = NewObject<USkeletalMeshBackedDynamicMeshComponentToolTarget>();
+	Target->InitializeComponent(Cast<UDynamicMeshComponent>(Component));
+	checkSlow(Target->Component.IsValid() && Requirements.AreSatisfiedBy(Target));
+
+	Target->bCommitToSkeletalMeshOnToolCommit = ComponentProvider->ShouldCommitToSkeletalMeshOnToolCommit();
+	
+	return Target;
+	
+}
+
+bool USkeletalMeshBackedDynamicMeshComponentToolTargetFactory::CanWriteToSource(const UObject* SourceObject)
+{
+	if (const USkeletalMeshComponent* Component = GetValid(Cast<USkeletalMeshComponent>(SourceObject)))
+	{
+		if (const USkeletalMesh* SkeletalMesh = GetValid(Cast<USkeletalMesh>(Component->GetSkeletalMeshAsset())))
+		{
+			return !SkeletalMesh->GetPathName().StartsWith(TEXT("/Engine/"));
+		}
+	}
+	return false;
+}
+
+#undef LOCTEXT_NAMESPACE
diff --git a/Engine/Plugins/Animation/SkeletalMeshModelingTools/Source/SkeletalMeshModelingTools/Private/ToolTargets/SKMBackedDynaMeshComponentToolTarget.h b/Engine/Plugins/Animation/SkeletalMeshModelingTools/Source/SkeletalMeshModelingTools/Private/ToolTargets/SKMBackedDynaMeshComponentToolTarget.h
new file mode 100644
--- /dev/null
+++ b/Engine/Plugins/Animation/SkeletalMeshModelingTools/Source/SkeletalMeshModelingTools/Private/ToolTargets/SKMBackedDynaMeshComponentToolTarget.h
@@ -0,0 +1,59 @@
+﻿// Copyright Epic Games, Inc. All Rights Reserved.
+
+#pragma once
+
+#include "CoreMinimal.h"
+#include "Components/DynamicMeshComponent.h"
+#include "Interfaces/ISKMBackedDynaMeshComponentProvider.h"
+#include "TargetInterfaces/SkeletalMeshBackedTarget.h"
+#include "ToolTargets/DynamicMeshComponentToolTarget.h"
+#include "UObject/WeakInterfacePtr.h"
+#include "ToolTargets/ToolTarget.h"
+
+#include "SKMBackedDynaMeshComponentToolTarget.generated.h"
+
+#define UE_API SKELETALMESHMODELINGTOOLS_API
+
+UCLASS(MinimalAPI, Transient)
+class USkeletalMeshBackedDynamicMeshComponentToolTarget :
+	public UDynamicMeshComponentToolTarget,
+	public ISkeletalMeshBackedTarget
+{
+	GENERATED_BODY()
+public:
+	virtual FReferenceSkeleton GetSkeleton() const override;
+	virtual void CommitSkeletonModifier(USkeletonModifier* InModifier) override;
+
+	virtual void CommitDynamicMeshChange(TUniquePtr<FToolCommandChange> Change, const FText& ChangeMessage) override;
+	
+	virtual USkeletalMesh* GetSkeletalMesh() const override;
+
+	virtual void SetOwnerVisibility(bool bVisible) const override;
+protected:
+	USkeletalMeshBackedDynamicMeshComponent* GetComponent() const;
+
+	TWeakObjectPtr<USkeletalMesh> WeakSkeletalMesh;
+
+	bool bCommitToSkeletalMeshOnToolCommit = true;
+	
+	friend class USkeletalMeshBackedDynamicMeshComponentToolTargetFactory;
+};
+
+UCLASS(MinimalAPI, Transient)
+class USkeletalMeshBackedDynamicMeshComponentToolTargetFactory : public UDynamicMeshComponentToolTargetFactory
+{
+	GENERATED_BODY()
+
+public:
+	void Init(ISkeletalMeshBackedDynamicMeshComponentProvider* InComponentProvider);
+
+	virtual bool CanBuildTarget(UObject* SourceObject, const FToolTargetTypeRequirements& TargetTypeInfo) const override;
+
+	virtual UToolTarget* BuildTarget(UObject* SourceObject, const FToolTargetTypeRequirements& TargetTypeInfo) override;
+
+	static bool CanWriteToSource(const UObject* SourceObject);
+
+	TWeakInterfacePtr<ISkeletalMeshBackedDynamicMeshComponentProvider> ComponentProvider;
+};
+
+#undef UE_API
diff --git a/Engine/Plugins/Animation/SkeletalMeshModelingTools/Source/SkeletalMeshModelingTools/SkeletalMeshModelingTools.Build.cs b/Engine/Plugins/Animation/SkeletalMeshModelingTools/Source/SkeletalMeshModelingTools/SkeletalMeshModelingTools.Build.cs
--- a/Engine/Plugins/Animation/SkeletalMeshModelingTools/Source/SkeletalMeshModelingTools/SkeletalMeshModelingTools.Build.cs
+++ b/Engine/Plugins/Animation/SkeletalMeshModelingTools/Source/SkeletalMeshModelingTools/SkeletalMeshModelingTools.Build.cs
@@ -42,7 +42,9 @@ public class SkeletalMeshModelingTools : ModuleRules
 				"GeometryCore",
 				"ModelingComponents",
 				"MeshConversion",
-				"InterchangeEngine"
+				"InterchangeEngine",
+				"GeometryFramework",
+				"EditorWidgets" 
 			}
 		);
 
diff --git a/Engine/Plugins/Experimental/Animation/SkeletalMeshMorphTargetEditingTools/Source/SkeletalMeshMorphTargetEditingTools/Private/IMorphTargetEditingToolInterface.cpp b/Engine/Plugins/Experimental/Animation/SkeletalMeshMorphTargetEditingTools/Source/SkeletalMeshMorphTargetEditingTools/Private/IMorphTargetEditingToolInterface.cpp
--- a/Engine/Plugins/Experimental/Animation/SkeletalMeshMorphTargetEditingTools/Source/SkeletalMeshMorphTargetEditingTools/Private/IMorphTargetEditingToolInterface.cpp
+++ b/Engine/Plugins/Experimental/Animation/SkeletalMeshMorphTargetEditingTools/Source/SkeletalMeshMorphTargetEditingTools/Private/IMorphTargetEditingToolInterface.cpp
@@ -9,7 +9,6 @@
 #include "MorphTargetEditingToolProperties.h"
 #include "PersonaModule.h"
 #include "SingleSelectionTool.h"
-#include "SKMMorphTargetBackedTarget.h"
 #include "SkeletalMesh/SkeletalMeshEditingInterface.h"
 
 #include UE_INLINE_GENERATED_CPP_BY_NAME(IMorphTargetEditingToolInterface)
@@ -19,35 +18,15 @@ void IMorphTargetEditingToolInterface::SetupMorphEditingToolCommon()
 	// Only support single selection tools at the moment
 	USingleSelectionTool* SingleTargetTool = CastChecked<USingleSelectionTool>(this);
 	ISkeletalMeshEditingInterface* EditingInterface = CastChecked<ISkeletalMeshEditingInterface>(this);
-	
-	if (USkeletalMeshEditorContextObjectBase* EditorContext = SingleTargetTool->GetToolManager()->GetContextObjectStore()->FindContext<USkeletalMeshEditorContextObjectBase>())
+	USkeletalMeshEditorContextObjectBase* EditorContext = SingleTargetTool->GetToolManager()->GetContextObjectStore()->FindContext<USkeletalMeshEditorContextObjectBase>();
+	if (EditorContext)
 	{
 		EditorContext->BindTo(EditingInterface);
 	}
 	
-	ISkeletalMeshMorphTargetBackedTarget* MorphTargetBackedTarget = Cast<ISkeletalMeshMorphTargetBackedTarget>(SingleTargetTool->GetTarget());
-
 	auto SetupFunction = [&](UMorphTargetEditingToolProperties* InProperties)
 	{
-		InProperties->MorphTargetNames = MorphTargetBackedTarget->GetEditableMorphTargetNames();
-		TArray<FName> SelectedMorphTargets = EditingInterface->GetSelectedMorphTargets();
-		
-		InProperties->EditMorphTargetName = SelectedMorphTargets.Num() > 0 ?
-			SelectedMorphTargets[0] : InProperties->MorphTargetNames.Num() > 0 ? InProperties->MorphTargetNames[0] : NAME_None;
-		
-		InProperties->NewMorphTargetName = MorphTargetBackedTarget->GetValidNameForNewMorphTarget(TEXT("NewMorphTarget"));
-
-		if (SelectedMorphTargets.Num() > 0)
-		{
-			InProperties->Operation = EMorphTargetEditorOperation::Edit;
-			MorphTargetBackedTarget->SetEditingMorphTargetName(InProperties->EditMorphTargetName);
-		}
-		else
-		{
-			InProperties->Operation = EMorphTargetEditorOperation::New;
-			MorphTargetBackedTarget->SetEditingMorphTargetName(InProperties->NewMorphTargetName);
-		}
-		
+		InProperties->EditMorphTargetName = EditorContext->GetEditingMorphTarget();
 	};
 	
 	SetupCommonProperties(SetupFunction);
@@ -59,7 +38,7 @@ void IMorphTargetEditingToolInterface::ShutdownMorphEditingToolCommon()
 
 	// Only support single selection tools at the moment
 	USingleSelectionTool* SingleTargetTool = CastChecked<USingleSelectionTool>(this);
-
+	
 	if (USkeletalMeshEditorContextObjectBase* EditorContext = SingleTargetTool->GetToolManager()->GetContextObjectStore()->FindContext<USkeletalMeshEditorContextObjectBase>())
 	{
 		EditorContext->UnbindFrom(EditingInterface);
diff --git a/Engine/Plugins/Experimental/Animation/SkeletalMeshMorphTargetEditingTools/Source/SkeletalMeshMorphTargetEditingTools/Private/MorphTargetEditingToolProperties.cpp b/Engine/Plugins/Experimental/Animation/SkeletalMeshMorphTargetEditingTools/Source/SkeletalMeshMorphTargetEditingTools/Private/MorphTargetEditingToolProperties.cpp
--- a/Engine/Plugins/Experimental/Animation/SkeletalMeshMorphTargetEditingTools/Source/SkeletalMeshMorphTargetEditingTools/Private/MorphTargetEditingToolProperties.cpp
+++ b/Engine/Plugins/Experimental/Animation/SkeletalMeshMorphTargetEditingTools/Source/SkeletalMeshMorphTargetEditingTools/Private/MorphTargetEditingToolProperties.cpp
@@ -4,12 +4,3 @@
 
 #include UE_INLINE_GENERATED_CPP_BY_NAME(MorphTargetEditingToolProperties)
 
-TArray<FName> UMorphTargetEditingToolProperties::GetMorphTargetNames()
-{
-	return MorphTargetNames;
-}
-
-FName UMorphTargetEditingToolProperties::GetEditingMorphTargetName()
-{
-	return Operation == EMorphTargetEditorOperation::New ? NewMorphTargetName : EditMorphTargetName;
-}
diff --git a/Engine/Plugins/Experimental/Animation/SkeletalMeshMorphTargetEditingTools/Source/SkeletalMeshMorphTargetEditingTools/Private/MorphTargetVertexSculptTool.cpp b/Engine/Plugins/Experimental/Animation/SkeletalMeshMorphTargetEditingTools/Source/SkeletalMeshMorphTargetEditingTools/Private/MorphTargetVertexSculptTool.cpp
--- a/Engine/Plugins/Experimental/Animation/SkeletalMeshMorphTargetEditingTools/Source/SkeletalMeshMorphTargetEditingTools/Private/MorphTargetVertexSculptTool.cpp
+++ b/Engine/Plugins/Experimental/Animation/SkeletalMeshMorphTargetEditingTools/Source/SkeletalMeshMorphTargetEditingTools/Private/MorphTargetVertexSculptTool.cpp
@@ -15,7 +15,6 @@
 #include "DynamicMeshToMeshDescription.h"
 #include "IPersonaEditorModeManager.h"
 #include "PersonaModule.h"
-#include "SKMMorphTargetBackedTarget.h"
 #include "SkeletalMeshOperations.h"
 #include "BaseGizmos/BrushStampIndicator.h"
 #include "TargetInterfaces/DynamicMeshCommitter.h"
@@ -24,12 +23,33 @@
 #include "AnimationRuntime.h"
 #include "EraseMorphTargetBrushOps.h"
 #include "SkeletalMeshAttributes.h"
+#include "DynamicMesh/DynamicVertexSkinWeightsAttribute.h"
+#include "SkeletalMesh/SkeletalMeshToolsHelper.h"
+#include "TargetInterfaces/MeshDescriptionProvider.h"
+#include "TargetInterfaces/SkeletonProvider.h"
+#include "DynamicMesh/MeshNormals.h"
 
 #include UE_INLINE_GENERATED_CPP_BY_NAME(MorphTargetVertexSculptTool)
 
 #define LOCTEXT_NAMESPACE "MorphTargetVertexSculptTool"
 
 
+bool UMorphTargetVertexSculptToolBuilder::CanBuildTool(const FToolBuilderState& SceneState) const
+{
+	USkeletalMeshEditorContextObjectBase* EditorContext = SceneState.ToolManager->GetContextObjectStore()->FindContext<USkeletalMeshEditorContextObjectBase>();
+	if (!EditorContext)
+	{
+		return false;
+	}
+
+	if (EditorContext->GetEditingMorphTarget() == NAME_None)
+	{
+		return false;
+	}
+	
+	return Super::CanBuildTool(SceneState);
+}
+
 UMeshSurfacePointTool* UMorphTargetVertexSculptToolBuilder::CreateNewTool(const FToolBuilderState& SceneState) const
 {
 	UMorphTargetVertexSculptTool* MorphTargetEditorTool = NewObject<UMorphTargetVertexSculptTool>(SceneState.ToolManager);
@@ -44,34 +64,128 @@ const FToolTargetTypeRequirements& UMorphTargetVertexSculptToolBuilder::GetTarge
 		UDynamicMeshProvider::StaticClass(),
 		UDynamicMeshCommitter::StaticClass(),
 		USceneComponentBackedTarget::StaticClass(),
-		USkeletalMeshMorphTargetBackedTarget::StaticClass(),
+		USkeletonProvider::StaticClass(),
 	});
 
 	return ToolRequirements;
 }
 
+FName UMorphTargetVertexSculptTool::GetEditingMorphTarget()
+{
+	if (EditorContext.IsValid())
+	{
+		return EditorContext->GetEditingMorphTarget();
+	}
+
+	return {};	
+}
+
+TMap<FName, float> UMorphTargetVertexSculptTool::GetMorphTargetWeights()
+{
+	if (EditorContext.IsValid())
+	{
+		return EditorContext->GetMorphTargetWeights();
+	}
+
+	return {};	
+}
+
+const TArray<FTransform>& UMorphTargetVertexSculptTool::GetComponentSpaceBoneTransforms()
+{
+	if (EditorContext.IsValid())
+	{
+		return EditorContext->GetComponentSpaceBoneTransforms(GetTarget());
+	}
+
+	if (DefaultComponentSpaceBoneTransforms.IsEmpty())
+	{
+		ISkeletonProvider* SkeletonProvider = CastChecked<ISkeletonProvider>(GetTarget());
+		SkeletonProvider->GetSkeleton().GetBoneAbsoluteTransforms(DefaultComponentSpaceBoneTransforms);
+	}
+	
+	return DefaultComponentSpaceBoneTransforms;
+}
+
+void UMorphTargetVertexSculptTool::ToggleBoneManipulation(bool bEnable)
+{
+	if (EditorContext.IsValid())
+	{
+		EditorContext->ToggleBoneManipulation(bEnable);
+	}
+}
+
+void UMorphTargetVertexSculptTool::RequestRefreshSculptMesh()
+{
+	bSculptMeshOutdated = true;
+}
+
+void UMorphTargetVertexSculptTool::HandlePoseChangeDetectorEvent(SkeletalMeshToolsHelper::FPoseChangeDetector::FPayload Payload)
+{
+	if (!AllowToolMeshUpdates())
+	{
+		return;
+	}
+	
+	using namespace SkeletalMeshToolsHelper;
+	if (Payload.CurrentState == FPoseChangeDetector::PoseStoppedChanging)
+	{
+		FMeshReplacementChange Change(MeshBeforePosing, MakeShared<FDynamicMesh3>(*GetSculptMesh()));
+		TGuardValue<bool> PosingSculptMeshScope(bPosingSculptMesh, true);
+		OnDynamicMeshComponentChanged(DynamicMeshComponent, &Change, false);	
+	}
+	else if (Payload.CurrentState == FPoseChangeDetector::PoseJustChanged ||
+		Payload.CurrentState == FPoseChangeDetector::PoseChanged)
+	{
+		if (Payload.CurrentState == FPoseChangeDetector::PoseJustChanged)
+		{
+			MeshBeforePosing = MakeShared<FDynamicMesh3>(*GetSculptMesh());		
+		}
+
+		RequestRefreshSculptMesh();
+	}
+}
+
+
 
 void UMorphTargetVertexSculptTool::Setup()
 {
+	EditorContext = GetToolManager()->GetContextObjectStore()->FindContext<USkeletalMeshEditorContextObjectBase>();
+	
 	SetupMorphEditingToolCommon();
 	
 	// Setup Vertex Sculpt Tool
 	Super::Setup();
 
-	MorphTargetBackedTarget = Cast<ISkeletalMeshMorphTargetBackedTarget>(GetTarget());
+	ViewProperties->MaterialMode = EMeshEditingMaterialModes::ExistingMaterial;
+
+	ToolDynamicMesh = NewObject<UDynamicMesh>();
 
-	OnToolMeshChangedDelegate = DynamicMeshComponent->OnMeshRegionChanged.AddUObject(this, &UMorphTargetVertexSculptTool::OnToolMeshChanged);
+	ToolDynamicMesh->SetMesh(*GetSculptMesh());
 
-	MeshWithoutCurrentMorph = *GetSculptMesh();
 	
-	InitializeCache();
+	
+	// Remove unused attributes to make undo/redo faster
+	GetSculptMesh()->Attributes()->RemoveAllMorphTargetAttributes();
+	GetSculptMesh()->Attributes()->RemoveAllSkinWeightsAttributes();
+	
+	constexpr bool bCopyNormals = false, bCopyColors = false, bCopyUVs = false, bCopyAttributes = false;
+	MeshWithoutEditingMorph.Copy(*GetSculptMesh(), bCopyNormals, bCopyColors, bCopyUVs, bCopyAttributes);
+	
+	ToolMorphTargetName = GetEditingMorphTarget();
+	
+	const FReferenceSkeleton& RefSkeleton = CastChecked<ISkeletonProvider>(GetTarget())->GetSkeleton();
+	RefSkeleton.GetBoneAbsoluteTransforms(ComponentSpaceTransformsRefPose);
+
+	ToggleBoneManipulation(true);
+
+	PoseChangeDetector.GetNotifier().AddUObject(this, &UMorphTargetVertexSculptTool::HandlePoseChangeDetectorEvent);
 }
 
 void UMorphTargetVertexSculptTool::RegisterBrushes()
 {
 	Super::RegisterBrushes();
 	
-	GetMeshWithoutCurrentMorphFunc = [this](){return &MeshWithoutCurrentMorph;};
+	GetMeshWithoutCurrentMorphFunc = [this](){return &MeshWithoutEditingMorph;};
 
 	// Had to hijack the EraseSculptLayer identifier from base mesh vertex sculpt tool for our erase morph target tool since it is the simplest way to get an icon for the tool.
 	RegisterBrushType(
@@ -83,74 +197,54 @@ void UMorphTargetVertexSculptTool::RegisterBrushes()
 
 void UMorphTargetVertexSculptTool::Shutdown(EToolShutdownType ShutdownType)
 {
-	if (DynamicMeshComponent != nullptr)
-	{
-		DynamicMeshComponent->OnMeshChanged.Remove(OnDynamicMeshComponentChangedHandle);
-	}
-
-	if (ShutdownType == EToolShutdownType::Accept)
-	{
-		UpdateCacheIfNeeded();
-    
-    	if (bHasValidData)
-    	{
-    		MorphTargetBackedTarget->SetEditingMorphTargetName(EditorToolProperties->GetEditingMorphTargetName());
-    		MorphTargetBackedTarget->SetDataToCommit(MoveTemp(ToolMeshDescription), ToolMorphTargetName);
-    	}
-    	else
-    	{
-    		GetToolManager()->DisplayMessage(LOCTEXT("MorphTargetInvalidData", "Morph Target is empty, no change was made") , EToolMessageLevel::UserNotification);
-    		ShutdownType = EToolShutdownType::Cancel;
-    	}	
-	}
-	
 	Super::Shutdown(ShutdownType);
 
 	ShutdownMorphEditingToolCommon();	
 }
 
-void UMorphTargetVertexSculptTool::OnTick(float DeltaTime)
+void UMorphTargetVertexSculptTool::CommitResult(UBaseDynamicMeshComponent* Component, bool bModifiedTopology)
 {
-	if (UPersonaEditorModeManagerContext* PersonaModeManagerContext = GetToolManager()->GetContextObjectStore()->FindContext<UPersonaEditorModeManagerContext>())
-	{
-		if (!PersonaModeManagerContext->GetPersonaEditorModeManager()->IsModeActive(FPersonaEditModes::SkeletonSelection))
-		{
-			PersonaModeManagerContext->GetPersonaEditorModeManager()->ActivateMode(FPersonaEditModes::SkeletonSelection);	
-		}
-	}
+	GetToolManager()->BeginUndoTransaction(LOCTEXT("MorphTargetSculptMeshToolTransactionName", "Sculpt Morph Target"));
 
-	Super::OnTick(DeltaTime);
+	ToolDynamicMesh->ProcessMesh([&](const FDynamicMesh3& InMesh)
+		{
+			// Commit the tool mesh cache instead of the sculpt mesh
+			UE::ToolTarget::CommitDynamicMeshUpdate(Target, InMesh, bModifiedTopology);
+		});
 
-	if (InStroke())
-	{
-		bCached = false;
-	}
-	else 
+	if (EditorContext.IsValid())
 	{
-		PoseToolMesh();		
+		EditorContext->NotifyMorphTargetEdited();
 	}
+	
+	GetToolManager()->EndUndoTransaction();	
 }
 
-void UMorphTargetVertexSculptTool::OnPropertyModified(UObject* PropertySet, FProperty* Property)
+void UMorphTargetVertexSculptTool::OnTick(float DeltaTime)
 {
-	Super::OnPropertyModified(PropertySet, Property);
+	Super::OnTick(DeltaTime);
 
-	if (Property->GetFName() == GET_MEMBER_NAME_CHECKED(UMorphTargetEditingToolProperties, NewMorphTargetName))
+	if (!InStroke())
 	{
-		EditorToolProperties->NewMorphTargetName = MorphTargetBackedTarget->GetValidNameForNewMorphTarget(EditorToolProperties->NewMorphTargetName);
-	}
-	else if (Property->GetFName() == GET_MEMBER_NAME_CHECKED(UMorphTargetEditingToolProperties, Operation))
-	{
-		if (EditorToolProperties->Operation == EMorphTargetEditorOperation::Edit)
+		const TArray<FTransform>& ComponentSpaceTransforms = GetComponentSpaceBoneTransforms();
+		const TMap<FName, float>& MorphTargetWeights = GetMorphTargetWeights();
+		
+		PoseChangeDetector.CheckPose(ComponentSpaceTransforms, MorphTargetWeights);
+
+		if (bSculptMeshOutdated)
 		{
-			if (EditorToolProperties->GetMorphTargetNames().Num() == 0)
-			{
-				EditorToolProperties->Operation = EMorphTargetEditorOperation::New;
-			}
+			bSculptMeshOutdated = false;
+			PoseSculptMesh(ComponentSpaceTransforms, MorphTargetWeights);
 		}
 	}
 }
 
+void UMorphTargetVertexSculptTool::OnPropertyModified(UObject* PropertySet, FProperty* Property)
+{
+	Super::OnPropertyModified(PropertySet, Property);
+}
+
+
 /*
  * internal Change classes
  */
@@ -178,6 +272,41 @@ void FMorphTargetVertexSculptNonSymmetricChange::Revert(UObject* Object)
 	}
 }
 
+void UMorphTargetVertexSculptTool::OnBeginStroke(const FRay& WorldRay)
+{
+	Super::OnBeginStroke(WorldRay);
+}
+
+class FMeshMorphTargetChange : public FMeshChange
+{
+public:
+	FName MorphTargetName;
+	TArray<int32> Vertices;
+	TArray<FVector> OldDeltas;
+	TArray<FVector> NewDeltas;
+
+	FMeshMorphTargetChange() = default;
+	
+	/** Describes this change (for debugging) */
+	virtual FString ToString() const override
+	{
+		return TEXT("Edited Morph Target");
+	}
+
+	virtual void ApplyChangeToMesh(UE::Geometry::FDynamicMesh3* Mesh, bool bRevert) const override
+	{
+		UE::Geometry::FDynamicMeshMorphTargetAttribute* MorphTargetAttribute = Mesh->Attributes()->GetMorphTargetAttribute(MorphTargetName);
+
+		ParallelFor(Vertices.Num(), [&](int32 Index)
+			{
+				const TArray<FVector>& Deltas = bRevert ? OldDeltas : NewDeltas;
+				MorphTargetAttribute->SetValue(Vertices[Index], Deltas[Index]); 		
+			});
+	};
+};
+
+
+
 void UMorphTargetVertexSculptTool::OnEndStroke()
 {
 	// update spatial
@@ -187,41 +316,82 @@ void UMorphTargetVertexSculptTool::OnEndStroke()
 
 	check(ActiveVertexChange);
 
-	TUniquePtr<TWrappedToolCommandChange<FMeshVertexChange>> NewChange = MakeUnique<TWrappedToolCommandChange<FMeshVertexChange>>();
-	NewChange->WrappedChange = MoveTemp(ActiveVertexChange->Change);
-	NewChange->BeforeModify = [this](bool bRevert)
-	{
-		this->WaitForPendingUndoRedoUpdate();
-		// Any sculpt change needs to be applied to the latest posed mesh during undo / redo, so pose the mesh immediately instead of waiting
-		// for the next tick
-		this->PoseToolMesh();
-	};
+	TMap<FName, float> MorphTargetWeights = MorphTargetWeightsForPosedMesh;
+	const TArray<FTransform>& ComponentSpaceTransforms = ComponentSpaceTransformsForPosedMesh; 
 
-	GetToolManager()->EmitObjectChange(DynamicMeshComponent, MoveTemp(NewChange), LOCTEXT("VertexSculptChange", "Brush Stroke"));
-	if (bMeshSymmetryIsValid && bApplySymmetry == false)
+	float EditingMorphTargetWeight = MorphTargetWeights[ToolMorphTargetName];
+	
+	// Don't update the cache if the morph cannot be extracted from the mesh
+	if (!FMath::IsNearlyZero(EditingMorphTargetWeight))
 	{
-		// if we end a stroke while symmetry is possible but disabled, we now have to assume that symmetry is no longer possible
-		GetToolManager()->EmitObjectChange(this, MakeUnique<FMorphTargetVertexSculptNonSymmetricChange>(), LOCTEXT("DisableSymmetryChange", "Disable Symmetry"));
-		bMeshSymmetryIsValid = false;
-		SymmetryProperties->bSymmetryCanBeEnabled = bMeshSymmetryIsValid;
+		FMeshMorphTargetChange* MorphTargetChange = new FMeshMorphTargetChange(); 
+		MorphTargetChange->MorphTargetName = ToolMorphTargetName;
+		MorphTargetChange->Vertices = MoveTemp(ActiveVertexChange->Change->Vertices);
+		MorphTargetChange->OldDeltas.SetNumUninitialized(MorphTargetChange->Vertices.Num());
+		MorphTargetChange->NewDeltas.SetNumUninitialized(MorphTargetChange->Vertices.Num());
+		
+		// Exclude the morph target we are trying to extract
+		MorphTargetWeights.Remove(ToolMorphTargetName);
+		
+		using namespace SkeletalMeshToolsHelper;
+		TArray<FMatrix> BoneMatrices = ComputeBoneMatrices(
+				ComponentSpaceTransformsRefPose,
+				ComponentSpaceTransforms
+				);
+
+		ToolDynamicMesh->EditMesh([&](FDynamicMesh3& Mesh)
+			{
+				FDynamicMeshMorphTargetAttribute* MorphTargetAttribute = Mesh.Attributes()->GetMorphTargetAttribute(ToolMorphTargetName);
+				
+				auto WriteFunction = [&](FVertInfo VertInfo, const FVector& UnposedVertPos)
+					{
+						int32 VertID = VertInfo.VertID;
+						FVector NewDelta = UnposedVertPos - FVector(Mesh.GetVertex(VertID));
+						NewDelta /= EditingMorphTargetWeight;
+						
+						FVector OldDelta;
+						MorphTargetAttribute->GetValue(VertID, OldDelta);
+						MorphTargetAttribute->SetValue(VertID, NewDelta);
+
+						MorphTargetChange->OldDeltas[VertInfo.VertArrayIndex] = OldDelta;
+						MorphTargetChange->NewDeltas[VertInfo.VertArrayIndex] = NewDelta;
+					};
+
+				GetUnposedMesh(WriteFunction, *GetSculptMesh(), Mesh, BoneMatrices, NAME_None, MorphTargetWeights, MorphTargetChange->Vertices);
+			});
+		
+		
+		TUniquePtr<TWrappedToolCommandChange<FMeshMorphTargetChange>> NewChange = MakeUnique<TWrappedToolCommandChange<FMeshMorphTargetChange>>();
+		
+		NewChange->WrappedChange.Reset(MorphTargetChange);
+		NewChange->BeforeModify = [this](bool bRevert)
+		{
+			this->WaitForPendingUndoRedoUpdate();
+		};
+		NewChange->AfterModify= [this](bool bRevert)
+		{
+			this->RequestRefreshSculptMesh();
+		};
+
+		GetToolManager()->EmitObjectChange(ToolDynamicMesh, MoveTemp(NewChange), LOCTEXT("VertexSculptChange", "Brush Stroke"));
+		if (bMeshSymmetryIsValid && bApplySymmetry == false)
+		{
+			// if we end a stroke while symmetry is possible but disabled, we now have to assume that symmetry is no longer possible
+			GetToolManager()->EmitObjectChange(this, MakeUnique<FMorphTargetVertexSculptNonSymmetricChange>(), LOCTEXT("DisableSymmetryChange", "Disable Symmetry"));
+			bMeshSymmetryIsValid = false;
+			SymmetryProperties->bSymmetryCanBeEnabled = bMeshSymmetryIsValid;
+		}
+
+		RequestRefreshSculptMesh();
 	}
+	
+	
 	LongTransactions.Close(GetToolManager());
 
 	delete ActiveVertexChange;
 	ActiveVertexChange = nullptr;
 }
 
-
-void UMorphTargetVertexSculptTool::OnToolMeshChanged(UDynamicMeshComponent* Component, const FMeshRegionChangeBase* Change, bool bRevert)
-{
-	if (!bPosingSculptMesh)
-	{
-		// If not posing, it must have been a sculpt change that was reverted/applied, should update the cache
-		bCached = false;
-	}
-}
-
-
 void UMorphTargetVertexSculptTool::SetupCommonProperties(const TFunction<void(UMorphTargetEditingToolProperties*)>& InSetupFunction)
 {
 	EditorToolProperties = NewObject<UMorphTargetEditingToolProperties>(this);
@@ -236,265 +406,48 @@ void UMorphTargetVertexSculptTool::HandleSkeletalMeshModified(const TArray<FName
 {
 }
 
-void UMorphTargetVertexSculptTool::InitializeCache()
+void UMorphTargetVertexSculptTool::PoseSculptMesh(const TArray<FTransform>& ComponentSpaceTransforms, const TMap<FName, float>& MorphTargetWeights)
 {
-	UPrimitiveComponent* Component = Cast<IPrimitiveComponentBackedTarget>(GetTarget())->GetOwnerComponent();
-	USkeletalMeshComponent* SkeletalMeshComponent = Cast<USkeletalMeshComponent>(Component);
-
-	ToolMeshDescription = *SkeletalMeshComponent->GetSkeletalMeshAsset()->GetMeshDescription(LodIndex0);
-
-	// Pre compact the mesh description to avoid compacting one every frame
-	if (ToolMeshDescription.NeedsCompact())
-	{
-		FElementIDRemappings Remappings;
-		ToolMeshDescription.Compact(Remappings);
-	}
-	else
-	{
-		// Make sure indexers are built before entering parallel work
-		ToolMeshDescription.BuildVertexIndexers();
-	}
-	
-	FSkeletalMeshAttributes Attributes(ToolMeshDescription);
+	using namespace SkeletalMeshToolsHelper;
 	
-	ToolMorphTargetName = EditorToolProperties->GetEditingMorphTargetName();
-	
-	if (EditorToolProperties->Operation == EMorphTargetEditorOperation::New)
-	{
-		check(!Attributes.GetMorphTargetNames().Contains(ToolMorphTargetName))
-		Attributes.RegisterMorphTargetAttribute(ToolMorphTargetName, false);
-	}
-
-	bCached = true;
-}
-
-void UMorphTargetVertexSculptTool::UpdateCacheIfNeeded()
-{
+	// have to wait for any outstanding stamp/undo update to finish...
 	WaitForPendingStampUpdateConst();
 	WaitForPendingUndoRedoUpdate();
-	UPrimitiveComponent* Component = Cast<IPrimitiveComponentBackedTarget>(GetTarget())->GetOwnerComponent();
-	USkeletalMeshComponent* SkeletalMeshComponent = Cast<USkeletalMeshComponent>(Component);
-
-	if (!SkeletalMeshComponent->GetSkeletalMeshAsset())
-	{
-		return;
-	}
-	
-	if (bCached)
-	{
-		return;
-	}
-
-	// Don't update the cache if the morph cannot be extracted from the mesh
-	if (FMath::IsNearlyZero(EditorToolProperties->MorphTargetWeight))
-	{
-		bCached = true;
-		return;
-	}
-
-	const FReferenceSkeleton& RefSkeleton = SkeletalMeshComponent->GetSkeletalMeshAsset()->GetRefSkeleton();
-	
-	TArray<FTransform> RefComponentSpaceTransforms;
-	for (int32 BoneIndex = 0; BoneIndex < RefSkeleton.GetNum(); ++BoneIndex)
-	{
-		// Need local to component transform
-		RefComponentSpaceTransforms.Add(FAnimationRuntime::GetComponentSpaceTransformRefPose(RefSkeleton, BoneIndex));
-	}
-
-
-	TMap<FName, float> MorphTargetWeights = PreviousMorphWeights;
-	MorphTargetWeights.Remove(ToolMorphTargetName);
-
-	// In the case that source mesh contains non-manifold verts, extra duplicated verts will be added to end of the dynamic mesh verts array
-	// see FMeshDescriptionToDynamicMesh::Convert
-	check(GetSculptMesh()->VertexCount() >= ToolMeshDescription.GetVertexPositions().GetNumElements());
-	for (FVertexID VertexID : ToolMeshDescription.Vertices().GetElementIDs())
-	{
-		ToolMeshDescription.GetVertexPositions()[VertexID] = FVector3f(GetSculptMesh()->GetVertex(VertexID));
-	}
-	
-	FMeshDescription* RefMeshDescription = SkeletalMeshComponent->GetSkeletalMeshAsset()->GetMeshDescription(LodIndex0);
-
-	const TArray<FTransform>& ComponentSpaceTransforms = PreviousPoseComponentSpace;
-
-	FSkeletalMeshAttributes Attributes(ToolMeshDescription);
-	TVertexAttributesRef<FVector3f> MorphTargetPosDeltaAttribute = Attributes.GetVertexMorphPositionDelta(ToolMorphTargetName);
-	
-	if (FSkeletalMeshOperations::GetUnposedMeshInPlace(ToolMeshDescription, *RefMeshDescription, RefComponentSpaceTransforms, ComponentSpaceTransforms, NAME_None, MorphTargetWeights))
-	{
-		bHasValidData = false;
-		for (FVertexID VertexID : ToolMeshDescription.Vertices().GetElementIDs())
-		{
-			FVector3f Delta = ToolMeshDescription.GetVertexPosition(VertexID) - RefMeshDescription->GetVertexPosition(VertexID);
-			Delta /= EditorToolProperties->MorphTargetWeight;
 
-			if (!bHasValidData && Delta.SizeSquared() > FMath::Square(UE_THRESH_POINTS_ARE_NEAR))
-			{
-				bHasValidData = true;
-			}
-			MorphTargetPosDeltaAttribute[VertexID] = Delta;
-		}
-	}
+	TArray<FMatrix> BoneMatrices = ComputeBoneMatrices(ComponentSpaceTransformsRefPose, ComponentSpaceTransforms);
 
-	bCached = true;
-}
+	FDynamicMesh3& SculptMesh = *GetSculptMesh();
 
-void UMorphTargetVertexSculptTool::PoseToolMesh()
-{
-	if (!AllowToolMeshUpdates())
-	{
-		return;
-	}
-	
-	auto IsMorphWeightChanged = [&PreviousMorphWeights = this->PreviousMorphWeights](const FName& InName, float InMorphWeight) -> bool
-	{
-		if (float* PreviousWeight = PreviousMorphWeights.Find(InName))
-		{
-			if (!FMath::IsNearlyEqual(*PreviousWeight, InMorphWeight))
-			{
-				return true;
-			}
-		}
-		else
+	ToolDynamicMesh->ProcessMesh([&](const FDynamicMesh3& InMesh)
 		{
-			return true;
-		}
-		return false;
-	};
-	
-	
-	UPrimitiveComponent* Component = Cast<IPrimitiveComponentBackedTarget>(GetTarget())->GetOwnerComponent();
-	USkeletalMeshComponent* SkeletalMeshComponent = Cast<USkeletalMeshComponent>(Component);
-	
-	TMap<FName, float> MorphTargetWeights;
-
-	for (const TPair<const UMorphTarget*, int32>& MorphTarget: SkeletalMeshComponent->ActiveMorphTargets)
-	{
-		const FName MorphName = MorphTarget.Key->GetFName();
-		const float MorphWeight = SkeletalMeshComponent->MorphTargetWeights[MorphTarget.Value];
+			FDynamicMeshMorphTargetAttribute* MorphTargetAttribute = InMesh.Attributes()->GetMorphTargetAttribute(ToolMorphTargetName);
 
-		MorphTargetWeights.Add(MorphName, MorphWeight);
-	}
-	
-	MorphTargetWeights.FindOrAdd(ToolMorphTargetName) = EditorToolProperties->MorphTargetWeight;
-
-	const TArray<FTransform>& ComponentSpaceTransforms = SkeletalMeshComponent->GetComponentSpaceTransforms();
-	
-	bool bPoseChanged = false;
+			float EditingMorphTargetWeight = MorphTargetWeights[ToolMorphTargetName];
 	
-	// First tick, force an update
-	if (!bPoseOpInitialized)
-	{
-		bPoseOpInitialized = true;
-		
-		PreviousMorphWeights = MorphTargetWeights;
-		PreviousPoseComponentSpace = ComponentSpaceTransforms;
-		bPoseChanged = true;
-	}
-	else
-	{
-		// Check for pose changes and update accordingly
-		
-		if (!bPoseChanged)
-		{
-			if (MorphTargetWeights.Num() != PreviousMorphWeights.Num())
-			{
-				bPoseChanged = true;
-			}
-		}
-		
-		if (!bPoseChanged)
-		{
-			for (const TPair<FName, float>& MorphWeight: MorphTargetWeights)
-			{
-				if (IsMorphWeightChanged(MorphWeight.Key, MorphWeight.Value))
-				{
-					bPoseChanged = true;
-					break;
-				}
-			}
-		}
-		
-		if (!bPoseChanged)
-		{
-			for (int32 BoneIndex=0; BoneIndex< ComponentSpaceTransforms.Num(); ++BoneIndex)
-			{
-				const FTransform& CurrentBoneTransform = ComponentSpaceTransforms[BoneIndex];
-				const FTransform& PrevBoneTransform = PreviousPoseComponentSpace[BoneIndex];
-				if (!CurrentBoneTransform.Equals(PrevBoneTransform))
+			auto WriteFunc = [&](int32 VertID, const FVector& PosedVertPos)
 				{
-					bPoseChanged = true;
-					break;
-				}
-			}
-		}
-	}
-
-	if (!bPoseChanged)
-	{
-		if (bPoseChangedLastTick)
-		{
-			FMeshReplacementChange Change(MeshBeforePosing, MakeShared<FDynamicMesh3>(*GetSculptMesh()));
-
-			TGuardValue<bool> PosingSculptMeshScope(bPosingSculptMesh, true);
-			OnDynamicMeshComponentChanged(DynamicMeshComponent, &Change, false);
-		}
-
-		bPoseChangedLastTick = bPoseChanged;
-		return;
-	}
-
-	if (!bPoseChangedLastTick)
-	{
-		MeshBeforePosing = MakeShared<FDynamicMesh3>(*GetSculptMesh());
-	}
-	bPoseChangedLastTick = bPoseChanged;
-	
-	
-	UpdateCacheIfNeeded();
-	
-	PreviousPoseComponentSpace = ComponentSpaceTransforms;
-	PreviousMorphWeights = MorphTargetWeights;
+					SculptMesh.SetVertex(VertID, PosedVertPos);
 
-	// have to wait for any outstanding stamp/undo update to finish...
-	WaitForPendingStampUpdateConst();
-	WaitForPendingUndoRedoUpdate();
+					FVector Delta;
+					MorphTargetAttribute->GetValue(VertID, Delta);
+			
+					MeshWithoutEditingMorph.SetVertex(VertID,  PosedVertPos - (Delta * EditingMorphTargetWeight));
+				};
 
-	TArray<FTransform> RefComponentSpaceTransforms;
-	for (int32 BoneIndex = 0; BoneIndex < ComponentSpaceTransforms.Num(); ++BoneIndex)
-	{
-		// Need local to component transform
-		RefComponentSpaceTransforms.Add(FAnimationRuntime::GetComponentSpaceTransformRefPose(SkeletalMeshComponent->GetSkeletalMeshAsset()->GetRefSkeleton(), BoneIndex));
-	}
+			GetPosedMesh(WriteFunc, InMesh, BoneMatrices, NAME_None, MorphTargetWeights);
+			FMeshNormals::QuickRecomputeOverlayNormals(SculptMesh);
 			
-	FMeshDescription* RefMeshDescription = SkeletalMeshComponent->GetSkeletalMeshAsset()->GetMeshDescription(LodIndex0);
-	// Reset verts to ref pose for posing
-	for (FVertexID VertexID : ToolMeshDescription.Vertices().GetElementIDs())
-	{
-		ToolMeshDescription.GetVertexPositions()[VertexID] = RefMeshDescription->GetVertexPositions()[VertexID];
-	}
+			ComponentSpaceTransformsForPosedMesh = ComponentSpaceTransforms;
+			MorphTargetWeightsForPosedMesh = MorphTargetWeights;
+		
 
-	// No need to compute normals for the mesh description, the sculpt dynamic mesh recomputes its own normals during OnDynamicMeshComponentChanged()
-	constexpr bool bSkipRecomputeNormalsTangents = true;
-	FSkeletalMeshOperations::GetPosedMeshInPlace(ToolMeshDescription, ComponentSpaceTransforms, NAME_None, MorphTargetWeights, bSkipRecomputeNormalsTangents);
+		});
 
-	FDynamicMesh3& Mesh = *GetSculptMesh();
-	
-	TArrayView<const FVector3f> VertexPositionsAttribute = ToolMeshDescription.GetVertexPositions().GetRawArray();
-	
-	FSkeletalMeshAttributes Attributes(ToolMeshDescription);
-	TVertexAttributesRef<FVector3f> MorphTargetPosDeltaAttribute = Attributes.GetVertexMorphPositionDelta(ToolMorphTargetName);
-	
-	for (int32 Index = 0; Index < VertexPositionsAttribute.Num(); ++Index)
-	{
-		Mesh.SetVertex(Index, FVector3d(VertexPositionsAttribute[Index]));
+	constexpr bool bNormal = true;
+	DynamicMeshComponent->FastNotifyPositionsUpdated(bNormal);
 
-		FVector3f Delta = MorphTargetPosDeltaAttribute[Index] * EditorToolProperties->MorphTargetWeight;
-		MeshWithoutCurrentMorph.SetVertex(Index, FVector3d(VertexPositionsAttribute[Index] - Delta));
-	}	
-	
-	DynamicMeshComponent->FastNotifyPositionsUpdated();
-}	
+
+}
 
 
 #undef LOCTEXT_NAMESPACE
diff --git a/Engine/Plugins/Experimental/Animation/SkeletalMeshMorphTargetEditingTools/Source/SkeletalMeshMorphTargetEditingTools/Private/SKMMorphTargetEditingToolsCommands.cpp b/Engine/Plugins/Experimental/Animation/SkeletalMeshMorphTargetEditingTools/Source/SkeletalMeshMorphTargetEditingTools/Private/SKMMorphTargetEditingToolsCommands.cpp
--- a/Engine/Plugins/Experimental/Animation/SkeletalMeshMorphTargetEditingTools/Source/SkeletalMeshMorphTargetEditingTools/Private/SKMMorphTargetEditingToolsCommands.cpp
+++ b/Engine/Plugins/Experimental/Animation/SkeletalMeshMorphTargetEditingTools/Source/SkeletalMeshMorphTargetEditingTools/Private/SKMMorphTargetEditingToolsCommands.cpp
@@ -21,7 +21,7 @@ void FSkeletalMeshMorphTargetEditingToolsCommands::RegisterCommands()
 {
 	UI_COMMAND(BeginMorphTargetTool, "Morph", "Edit Morph Targets", EUserInterfaceActionType::ToggleButton, FInputChord());
 	
-	UI_COMMAND(BeginMorphTargetSculptTool, "Sculpt Morph Target", "MorphTarget Sculpt Tool", EUserInterfaceActionType::ToggleButton, FInputChord());
+	UI_COMMAND(BeginMorphTargetSculptTool, "Sculpt Morph Target", "To use this tool, select a Morph Target for editing from the morph target tab using the radio button", EUserInterfaceActionType::ToggleButton, FInputChord());
 }
 
 #undef LOCTEXT_NAMESPACE
diff --git a/Engine/Plugins/Experimental/Animation/SkeletalMeshMorphTargetEditingTools/Source/SkeletalMeshMorphTargetEditingTools/Private/SKMMorphTargetEditingToolsModule.cpp b/Engine/Plugins/Experimental/Animation/SkeletalMeshMorphTargetEditingTools/Source/SkeletalMeshMorphTargetEditingTools/Private/SKMMorphTargetEditingToolsModule.cpp
--- a/Engine/Plugins/Experimental/Animation/SkeletalMeshMorphTargetEditingTools/Source/SkeletalMeshMorphTargetEditingTools/Private/SKMMorphTargetEditingToolsModule.cpp
+++ b/Engine/Plugins/Experimental/Animation/SkeletalMeshMorphTargetEditingTools/Source/SkeletalMeshMorphTargetEditingTools/Private/SKMMorphTargetEditingToolsModule.cpp
@@ -5,7 +5,6 @@
 #include "MorphTargetVertexSculptTool.h"
 #include "SKMMorphTargetEditingToolsCommands.h"
 #include "SKMMorphTargetEditingToolsStyle.h"
-#include "SKMMorphTargetToolTarget.h"
 #include "Features/IModularFeatures.h"
 #include "Framework/Commands/Commands.h"
 
@@ -50,12 +49,12 @@ bool FSkeletalMeshMorphTargetEditingToolsModule::GetExtensionExtendedInfo(FModel
 
 bool FSkeletalMeshMorphTargetEditingToolsModule::GetExtensionToolTargets(TArray<TSubclassOf<UToolTargetFactory>>& ToolTargetFactoriesOut)
 {
-	ToolTargetFactoriesOut.Add(USkeletalMeshMorphTargetToolTargetFactory::StaticClass());
-	
-	return true;
+	return false;
 }
 
 
+
+
 #undef LOCTEXT_NAMESPACE
 	
 IMPLEMENT_MODULE(FSkeletalMeshMorphTargetEditingToolsModule, SkeletalMeshMorphTargetEditingTools)
\ No newline at end of file
diff --git a/Engine/Plugins/Experimental/Animation/SkeletalMeshMorphTargetEditingTools/Source/SkeletalMeshMorphTargetEditingTools/Private/SKMMorphTargetToolTarget.cpp b/Engine/Plugins/Experimental/Animation/SkeletalMeshMorphTargetEditingTools/Source/SkeletalMeshMorphTargetEditingTools/Private/SKMMorphTargetToolTarget.cpp
deleted file mode 100644
diff --git a/Engine/Plugins/Experimental/Animation/SkeletalMeshMorphTargetEditingTools/Source/SkeletalMeshMorphTargetEditingTools/Public/IMorphTargetEditingToolInterface.h b/Engine/Plugins/Experimental/Animation/SkeletalMeshMorphTargetEditingTools/Source/SkeletalMeshMorphTargetEditingTools/Public/IMorphTargetEditingToolInterface.h
--- a/Engine/Plugins/Experimental/Animation/SkeletalMeshMorphTargetEditingTools/Source/SkeletalMeshMorphTargetEditingTools/Public/IMorphTargetEditingToolInterface.h
+++ b/Engine/Plugins/Experimental/Animation/SkeletalMeshMorphTargetEditingTools/Source/SkeletalMeshMorphTargetEditingTools/Public/IMorphTargetEditingToolInterface.h
@@ -7,7 +7,6 @@
 
 #include "IMorphTargetEditingToolInterface.generated.h"
 
-class FMorphTargetEditChange;
 class UMorphTargetEditingToolProperties;
 
 UINTERFACE(MinimalAPI)
@@ -23,7 +22,5 @@ public:
 	virtual void SetupMorphEditingToolCommon();
 	virtual void ShutdownMorphEditingToolCommon();
 	virtual void SetupCommonProperties(const TFunction<void(UMorphTargetEditingToolProperties*)>& InSetupFunction) = 0;
-	// virtual void Apply(const FMorphTargetEditChange& Change) = 0;
-	// virtual void Revert(const FMorphTargetEditChange& Change) = 0;
 };
 
diff --git a/Engine/Plugins/Experimental/Animation/SkeletalMeshMorphTargetEditingTools/Source/SkeletalMeshMorphTargetEditingTools/Public/MorphTargetEditingToolProperties.h b/Engine/Plugins/Experimental/Animation/SkeletalMeshMorphTargetEditingTools/Source/SkeletalMeshMorphTargetEditingTools/Public/MorphTargetEditingToolProperties.h
--- a/Engine/Plugins/Experimental/Animation/SkeletalMeshMorphTargetEditingTools/Source/SkeletalMeshMorphTargetEditingTools/Public/MorphTargetEditingToolProperties.h
+++ b/Engine/Plugins/Experimental/Animation/SkeletalMeshMorphTargetEditingTools/Source/SkeletalMeshMorphTargetEditingTools/Public/MorphTargetEditingToolProperties.h
@@ -4,41 +4,13 @@
 
 #include "MorphTargetEditingToolProperties.generated.h"
 
-UENUM()
-enum class EMorphTargetEditorOperation
-{
-	// Add new MorphTarget, select a morph target before activating the tool to edit
-	New,
-	// Edit existing MorphTarget, deselect morph targets before activating the tool to create a new morph target
-	Edit, 
-};
+
 
 UCLASS()
 class UMorphTargetEditingToolProperties : public UInteractiveToolPropertySet
 {
 	GENERATED_BODY()
 public:
-
-	// Select a morph target before activating the tool to edit, if nothing is selected, a new morph target is created
-	UPROPERTY(VisibleAnywhere, Category = "MorphTarget")
-	EMorphTargetEditorOperation Operation;
-
-	UPROPERTY(EditAnywhere, Category = "MorphTarget", meta = (DisplayName = "Name", EditCondition = "Operation == EMorphTargetEditorOperation::New", EditConditionHides))
-	FName NewMorphTargetName;
-	
-	UPROPERTY(VisibleAnywhere, Category = "MorphTarget", meta = (GetOptions = GetMorphTargetNames, NoResetToDefault, DisplayName = "Name", EditCondition = "Operation == EMorphTargetEditorOperation::Edit", EditConditionHides))
+	UPROPERTY(VisibleAnywhere, Category = "MorphTarget", meta = (GetOptions = GetMorphTargetNames, NoResetToDefault, DisplayName = "Name"))
 	FName EditMorphTargetName;
-
-	UPROPERTY(EditAnywhere, Category = "MorphTarget", meta = (DisplayName = "Weight", UIMin = "0.0", UIMax = "1.0"))
-	float MorphTargetWeight = 1.f;
-
-	UFUNCTION()
-	TArray<FName> GetMorphTargetNames();
-
-	UPROPERTY()
-	TArray<FName> MorphTargetNames;
-
-	FName GetEditingMorphTargetName();
-
-	
 };
\ No newline at end of file
diff --git a/Engine/Plugins/Experimental/Animation/SkeletalMeshMorphTargetEditingTools/Source/SkeletalMeshMorphTargetEditingTools/Public/MorphTargetVertexSculptTool.h b/Engine/Plugins/Experimental/Animation/SkeletalMeshMorphTargetEditingTools/Source/SkeletalMeshMorphTargetEditingTools/Public/MorphTargetVertexSculptTool.h
--- a/Engine/Plugins/Experimental/Animation/SkeletalMeshMorphTargetEditingTools/Source/SkeletalMeshMorphTargetEditingTools/Public/MorphTargetVertexSculptTool.h
+++ b/Engine/Plugins/Experimental/Animation/SkeletalMeshMorphTargetEditingTools/Source/SkeletalMeshMorphTargetEditingTools/Public/MorphTargetVertexSculptTool.h
@@ -7,8 +7,9 @@
 #include "SkeletalMesh/SkeletalMeshEditingInterface.h"
 #include "BaseTools/MultiSelectionMeshEditingTool.h"
 #include "MorphTargetEditingToolProperties.h"
-#include "SKMMorphTargetBackedTarget.h"
 #include "MeshDescription.h"
+#include "ReferenceSkeleton.h"
+#include "SkeletalMesh/SkeletalMeshToolsHelper.h"
 
 #include "MorphTargetVertexSculptTool.generated.h"
 
@@ -21,13 +22,14 @@ class UMorphTargetVertexSculptToolBuilder : public UMeshVertexSculptToolBuilder
 	GENERATED_BODY()
 
 public:
+	virtual bool CanBuildTool(const FToolBuilderState& SceneState) const override;
 	virtual UMeshSurfacePointTool* CreateNewTool(const FToolBuilderState& SceneState) const override;
 	virtual const FToolTargetTypeRequirements& GetTargetRequirements() const override;
 };
 
 
 /**
-* MorphTarget Editor Tool
+* MorphTarget Sculpt Tool
 */
 UCLASS()
 class UMorphTargetVertexSculptTool : 
@@ -37,58 +39,69 @@ class UMorphTargetVertexSculptTool :
 {
 	GENERATED_BODY()
 
-public:	
+public:
+	virtual FName GetEditingMorphTarget();
+	virtual TMap<FName, float> GetMorphTargetWeights();
+	virtual const TArray<FTransform>& GetComponentSpaceBoneTransforms();
+	virtual void ToggleBoneManipulation(bool bEnable);
+
+	void RequestRefreshSculptMesh();
+
+
 	// UMeshVertexSculptTool overrides
 	virtual void Setup() override;
 	virtual void Shutdown(EToolShutdownType ShutdownType) override;
+	virtual void CommitResult(UBaseDynamicMeshComponent* Component, bool bModifiedTopology) override;
 	virtual void OnTick(float DeltaTime) override;
 	virtual void OnPropertyModified(UObject* PropertySet, FProperty* Property) override;
+	virtual void OnBeginStroke(const FRay& WorldRay) override;
 	virtual void OnEndStroke() override;
 
-	
 protected:
 	friend class FMorphTargetVertexSculptNonSymmetricChange;
 	virtual void RegisterBrushes() override;
 
-	void OnToolMeshChanged(UDynamicMeshComponent* Component, const FMeshRegionChangeBase* Change, bool bRevert);
+	void HandlePoseChangeDetectorEvent(SkeletalMeshToolsHelper::FPoseChangeDetector::FPayload Payload);
 
 	// IMorphTargetEditingTool overrides
 	virtual void SetupCommonProperties(const TFunction<void(UMorphTargetEditingToolProperties*)>& InSetupFunction) override;
 	
 	// ISkeletalMeshEditingInterface
 	virtual void HandleSkeletalMeshModified(const TArray<FName>& Payload, const ESkeletalMeshNotifyType InNotifyType) override;
-
-	void InitializeCache();
-	void UpdateCacheIfNeeded();
-
-	void PoseToolMesh();
+	
+	
+	void PoseSculptMesh(const TArray<FTransform>& ComponentSpaceTransforms, const TMap<FName, float>& MorphTargetWeights);
 
 	UPROPERTY()
 	TObjectPtr<class UMorphTargetEditingToolProperties> EditorToolProperties;
 
-	TWeakInterfacePtr<ISkeletalMeshMorphTargetBackedTarget> MorphTargetBackedTarget; 
+	UPROPERTY()
+	TWeakObjectPtr<USkeletalMeshEditorContextObjectBase> EditorContext = nullptr;
 
 	TFunction<const FDynamicMesh3*()> GetMeshWithoutCurrentMorphFunc;
-	FDynamicMesh3 MeshWithoutCurrentMorph;
+	FDynamicMesh3 MeshWithoutEditingMorph;
 	
-	bool bCached = false;
-	bool bHasValidData = false;
-
-	bool bPoseOpInitialized = false;
-	TArray<FTransform> PreviousPoseComponentSpace;
-	TMap<FName, float> PreviousMorphWeights;
+	SkeletalMeshToolsHelper::FPoseChangeDetector PoseChangeDetector;
 
 	TSharedPtr<FDynamicMesh3> MeshBeforePosing;
-	bool bPoseChangedLastTick = false;
 
 	FDelegateHandle OnToolMeshChangedDelegate;
 
-	static constexpr int32 LodIndex0 = 0;
-
-	FMeshDescription ToolMeshDescription;
+	UPROPERTY()
+	TObjectPtr<UDynamicMesh> ToolDynamicMesh;
+	
+	TArray<FTransform> ComponentSpaceTransformsRefPose;
+	
 	FName ToolMorphTargetName;
 
 	bool bPosingSculptMesh = false;
+
+	TArray<FTransform> DefaultComponentSpaceBoneTransforms;
+	
+	TArray<FTransform> ComponentSpaceTransformsForPosedMesh;
+	TMap<FName, float> MorphTargetWeightsForPosedMesh;
+
+	bool bSculptMeshOutdated = false;
 };
 
 
diff --git a/Engine/Plugins/Experimental/Animation/SkeletalMeshMorphTargetEditingTools/Source/SkeletalMeshMorphTargetEditingTools/Public/SKMMorphTargetBackedTarget.h b/Engine/Plugins/Experimental/Animation/SkeletalMeshMorphTargetEditingTools/Source/SkeletalMeshMorphTargetEditingTools/Public/SKMMorphTargetBackedTarget.h
deleted file mode 100644
diff --git a/Engine/Plugins/Experimental/Animation/SkeletalMeshMorphTargetEditingTools/Source/SkeletalMeshMorphTargetEditingTools/Public/SKMMorphTargetToolTarget.h b/Engine/Plugins/Experimental/Animation/SkeletalMeshMorphTargetEditingTools/Source/SkeletalMeshMorphTargetEditingTools/Public/SKMMorphTargetToolTarget.h
deleted file mode 100644
diff --git a/Engine/Plugins/Experimental/Animation/SkeletalMeshMorphTargetEditingTools/Source/SkeletalMeshMorphTargetEditingTools/SkeletalMeshMorphTargetEditingTools.Build.cs b/Engine/Plugins/Experimental/Animation/SkeletalMeshMorphTargetEditingTools/Source/SkeletalMeshMorphTargetEditingTools/SkeletalMeshMorphTargetEditingTools.Build.cs
--- a/Engine/Plugins/Experimental/Animation/SkeletalMeshMorphTargetEditingTools/Source/SkeletalMeshMorphTargetEditingTools/SkeletalMeshMorphTargetEditingTools.Build.cs
+++ b/Engine/Plugins/Experimental/Animation/SkeletalMeshMorphTargetEditingTools/Source/SkeletalMeshMorphTargetEditingTools/SkeletalMeshMorphTargetEditingTools.Build.cs
@@ -28,7 +28,8 @@ public class SkeletalMeshMorphTargetEditingTools : ModuleRules
 				"Core",
 				// ... add other public dependencies that you statically link with here ...
 				"ModelingToolsEditorMode",
-				"SkeletalMeshModelingTools"
+				"SkeletalMeshModelingTools", 
+				"MeshModelingToolsEditorOnly"
 			}
 			);
 			
diff --git a/Engine/Plugins/Experimental/Dataflow/Source/DataflowEditor/Private/DataflowEditorTools/DataflowEditorEditSkeletonBonesTool.cpp b/Engine/Plugins/Experimental/Dataflow/Source/DataflowEditor/Private/DataflowEditorTools/DataflowEditorEditSkeletonBonesTool.cpp
--- a/Engine/Plugins/Experimental/Dataflow/Source/DataflowEditor/Private/DataflowEditorTools/DataflowEditorEditSkeletonBonesTool.cpp
+++ b/Engine/Plugins/Experimental/Dataflow/Source/DataflowEditor/Private/DataflowEditorTools/DataflowEditorEditSkeletonBonesTool.cpp
@@ -82,12 +82,11 @@ UInteractiveTool* UDataflowEditorEditSkeletonBonesToolBuilder::BuildTool(const F
 			TSharedPtr<FDataflowSkeletalMeshEditorBinding> SkeletalMeshBinding = MakeShared<FDataflowSkeletalMeshEditorBinding>();
 			
 			EditTool->EditSkeletonBonesNode = EditSkeletonBonesNode;
-			EditTool->BindTo(SkeletalMeshBinding);
 			
 			EditSkeletonBonesNode->OnBoneSelectionChanged.AddLambda(
 				[EditTool, SkeletalMeshBinding](const TArray<FName>& BoneNames)
 			{
-				EditTool->GetNotifier().HandleNotification(BoneNames, ESkeletalMeshNotifyType::BonesSelected);
+				EditTool->GetNotifier()->HandleNotification(BoneNames, ESkeletalMeshNotifyType::BonesSelected);
 				SkeletalMeshBinding->BoneSelection = BoneNames;
 			});
 			EditTool->SetDataflowEditorContextObject(ContextObject);
@@ -120,12 +119,15 @@ void UDataflowEditorEditSkeletonBonesTool::Shutdown(EToolShutdownType ShutdownTy
 }
 
 FDataflowSkeletalMeshEditorBinding::FDataflowSkeletalMeshEditorBinding()
-	: ISkeletalMeshEditorBinding(),
-	DataflowNotifier()
+	: ISkeletalMeshEditorBinding()
 {}
 
-ISkeletalMeshNotifier& FDataflowSkeletalMeshEditorBinding::GetNotifier()
+TSharedPtr<ISkeletalMeshNotifier> FDataflowSkeletalMeshEditorBinding::GetNotifier()
 {
+	if (!DataflowNotifier)
+	{
+		DataflowNotifier = MakeShared<FDataflowSkeletalMeshEditorNotifier>();
+	}
 	return DataflowNotifier;	
 }
 
@@ -148,12 +150,6 @@ TArray<FName> FDataflowSkeletalMeshEditorBinding::GetSelectedBones() const
 	return BoneSelection;
 }
 
-TArray<FName> FDataflowSkeletalMeshEditorBinding::GetSelectedMorphTargets() const
-{
-	TArray<FName> Selection;
-	return Selection;
-}
-
 FDataflowSkeletalMeshEditorNotifier::FDataflowSkeletalMeshEditorNotifier()
 	: ISkeletalMeshNotifier()
 {}
diff --git a/Engine/Plugins/Experimental/Dataflow/Source/DataflowEditor/Private/DataflowEditorTools/DataflowEditorSkinWeightsPaintTool.cpp b/Engine/Plugins/Experimental/Dataflow/Source/DataflowEditor/Private/DataflowEditorTools/DataflowEditorSkinWeightsPaintTool.cpp
--- a/Engine/Plugins/Experimental/Dataflow/Source/DataflowEditor/Private/DataflowEditorTools/DataflowEditorSkinWeightsPaintTool.cpp
+++ b/Engine/Plugins/Experimental/Dataflow/Source/DataflowEditor/Private/DataflowEditorTools/DataflowEditorSkinWeightsPaintTool.cpp
@@ -107,7 +107,7 @@ void UDataflowEditorSkinWeightsPaintTool::SetSkinWeightNode(FDataflowCollectionE
 
 void UDataflowEditorSkinWeightsPaintTool::HandleOnBoneSelectionChanged(const TArray<FName>& BoneNames)
 {
-	GetNotifier().HandleNotification(BoneNames, ESkeletalMeshNotifyType::BonesSelected);
+	GetNotifier()->HandleNotification(BoneNames, ESkeletalMeshNotifyType::BonesSelected);
 }
 
 void UDataflowEditorSkinWeightsPaintTool::SetDataflowEditorContextObject(TObjectPtr<UDataflowContextObject> InDataflowEditorContextObject)
diff --git a/Engine/Plugins/Experimental/Dataflow/Source/DataflowEditor/Public/DataflowEditorTools/DataflowEditorEditSkeletonBonesTool.h b/Engine/Plugins/Experimental/Dataflow/Source/DataflowEditor/Public/DataflowEditorTools/DataflowEditorEditSkeletonBonesTool.h
--- a/Engine/Plugins/Experimental/Dataflow/Source/DataflowEditor/Public/DataflowEditorTools/DataflowEditorEditSkeletonBonesTool.h
+++ b/Engine/Plugins/Experimental/Dataflow/Source/DataflowEditor/Public/DataflowEditorTools/DataflowEditorEditSkeletonBonesTool.h
@@ -87,10 +87,9 @@ public:
 	FDataflowSkeletalMeshEditorBinding();
 
 	//~ Begin ISkeletalMeshEditorBinding interface
-	virtual ISkeletalMeshNotifier& GetNotifier() override;
+	virtual TSharedPtr<ISkeletalMeshNotifier> GetNotifier() override;
 	virtual NameFunction GetNameFunction() override;
 	virtual TArray<FName> GetSelectedBones() const override;
-	virtual TArray<FName> GetSelectedMorphTargets() const override;
 	//~ End ISkeletalMeshEditorBinding interface
 
 	TArray<FName> BoneSelection;
@@ -98,7 +97,7 @@ public:
 private :
 
 	/** Notifier used when the tool selection is changing */
-	FDataflowSkeletalMeshEditorNotifier DataflowNotifier;
+	TSharedPtr<FDataflowSkeletalMeshEditorNotifier> DataflowNotifier;
 };
 
 
diff --git a/Engine/Plugins/Runtime/MeshModelingToolset/Source/MeshModelingToolsEditorOnly/Private/SkeletalMesh/RefSkeletonPoser.cpp b/Engine/Plugins/Runtime/MeshModelingToolset/Source/MeshModelingToolsEditorOnly/Private/SkeletalMesh/RefSkeletonPoser.cpp
new file mode 100644
--- /dev/null
+++ b/Engine/Plugins/Runtime/MeshModelingToolset/Source/MeshModelingToolsEditorOnly/Private/SkeletalMesh/RefSkeletonPoser.cpp
@@ -0,0 +1,192 @@
+﻿// Copyright Epic Games, Inc. All Rights Reserved.
+
+
+#include "SkeletalMesh/RefSkeletonPoser.h"
+#include "Misc/ITransaction.h"
+#include "Editor/EditorEngine.h"
+
+extern UNREALED_API UEditorEngine* GEditor;
+
+#include UE_INLINE_GENERATED_CPP_BY_NAME(RefSkeletonPoser)
+
+#define LOCTEXT_NAMESPACE "RefSkeletonPoser"
+
+URefSkeletonPoser::FPoseChange::FPoseChange(URefSkeletonPoser* InPoser)
+{
+	OldAdditiveTransforms = InPoser->AdditiveTransforms;
+	NewAdditiveTransforms = OldAdditiveTransforms;
+}
+
+FString URefSkeletonPoser::FPoseChange::ToString() const
+{
+	return FString(TEXT("Pose Skeleton"));
+}
+
+void URefSkeletonPoser::FPoseChange::Apply(UObject* Object)
+{
+	if (URefSkeletonPoser* Poser = Cast<URefSkeletonPoser>(Object))
+	{
+		Poser->AdditiveTransforms = NewAdditiveTransforms;
+		Poser->Invalidate();
+	}
+}
+
+void URefSkeletonPoser::FPoseChange::Revert(UObject* Object)
+{
+	if (URefSkeletonPoser* Poser = CastChecked<URefSkeletonPoser>(Object))
+	{
+		Poser->AdditiveTransforms = OldAdditiveTransforms;
+		Poser->Invalidate();
+	}
+}
+
+void URefSkeletonPoser::FPoseChange::Close(URefSkeletonPoser* InPoser)
+{
+	NewAdditiveTransforms = InPoser->AdditiveTransforms;
+}
+
+void URefSkeletonPoser::SetRefSkeleton(const FReferenceSkeleton& InRefSkeleton)
+{
+	RefSkeleton = InRefSkeleton;
+	
+	Transforms = RefSkeleton.GetRefBonePose();
+	TransformCached = TBitArray<>(false, RefSkeleton.GetRawBoneNum());
+	
+	RefSkeleton.GetBoneAbsoluteTransforms(RefTransforms);
+
+	AdditiveTransforms.Reset();
+}
+
+void URefSkeletonPoser::ModifyBoneAdditiveTransform(int32 BoneIndex, TFunctionRef<void(FTransform&)> InModifyFunc)
+{
+	FTransform& TransformRef = AdditiveTransforms.FindOrAdd(BoneIndex);
+
+	FTransform NewTransform = TransformRef;
+	InModifyFunc(NewTransform);
+	
+
+	if (!TransformRef.Equals(NewTransform))
+	{
+		TransformRef = NewTransform;
+		bPoseChanged = true;
+	}
+	Invalidate(BoneIndex);
+}
+
+void URefSkeletonPoser::ClearBoneAdditiveTransform(int32 BoneIndex)
+{
+	if (AdditiveTransforms.Contains(BoneIndex))
+	{
+		AdditiveTransforms.Remove(BoneIndex);
+		bPoseChanged = true;
+		Invalidate(BoneIndex);
+	}
+}
+
+void URefSkeletonPoser::ClearAllBoneAdditiveTransforms()
+{
+	if (!AdditiveTransforms.IsEmpty())
+	{
+		AdditiveTransforms.Reset();
+		bPoseChanged = true;
+		Invalidate();
+	}
+}
+
+TOptional<FTransform> URefSkeletonPoser::GetBoneAdditiveTransform(int32 BoneIndex)
+{
+	if (FTransform* Transform = AdditiveTransforms.Find(BoneIndex))
+	{
+		return *Transform;
+	}
+			
+	return {};
+}
+
+const TArray<FTransform>& URefSkeletonPoser::GetComponentSpaceTransforms()
+{
+	for (int32 BoneIndex = 0; BoneIndex < RefSkeleton.GetNum(); ++BoneIndex)
+	{
+		GetComponentSpaceTransform(BoneIndex);
+	}
+			
+	return Transforms;
+}
+
+const FTransform& URefSkeletonPoser::GetComponentSpaceTransform(const uint32 Index)
+{
+	if (!Transforms.IsValidIndex(Index))
+	{
+		return FTransform::Identity;
+	}
+
+	if (TransformCached[Index])
+	{
+		return Transforms[Index];
+	}
+
+	if (FTransform* Transform = AdditiveTransforms.Find(Index))
+	{
+		Transforms[Index] = *Transform * Transforms[Index];
+	}
+
+	const int32 ParentIndex = RefSkeleton.GetParentIndex(Index);
+	if (ParentIndex > INDEX_NONE)
+	{
+		Transforms[Index] *= GetComponentSpaceTransform(ParentIndex);
+	}
+
+	TransformCached[Index] = true;
+	return Transforms[Index];
+}
+
+void URefSkeletonPoser::Invalidate(const uint32 Index)
+{
+	if (!TransformCached.IsValidIndex(Index))
+	{
+		Transforms = RefSkeleton.GetRefBonePose();
+		TransformCached.Init(false, Transforms.Num());
+		return;
+	}
+			
+	Transforms[Index] = RefSkeleton.GetRefBonePose()[Index];
+	TransformCached[Index] = false;
+
+	TArray<int32> Children; RefSkeleton.GetDirectChildBones(Index, Children);
+	
+	for (const int32 ChildIndex: Children)
+	{
+		Invalidate(ChildIndex);
+	}
+}
+
+const TArray<FTransform>& URefSkeletonPoser::GetComponentSpaceTransformsRefPose()
+{
+	return RefTransforms;
+}
+
+void URefSkeletonPoser::BeginPoseChange()
+{
+	GEditor->BeginTransaction(LOCTEXT("ChangePose", "Change Pose"));
+	ActivePoseChange.Reset(new FPoseChange(this));
+	bPoseChanged = false;
+}
+
+void URefSkeletonPoser::EndPoseChange()
+{
+	if (bPoseChanged)
+	{
+		ActivePoseChange->Close(this);
+		GUndo->StoreUndo(this, MoveTemp(ActivePoseChange));
+	}
+	ActivePoseChange.Reset();
+	GEditor->EndTransaction();
+}
+
+bool URefSkeletonPoser::IsRecordingPoseChange()
+{
+	return ActivePoseChange.IsValid();
+}
+
+
+#undef LOCTEXT_NAMESPACE
diff --git a/Engine/Plugins/Runtime/MeshModelingToolset/Source/MeshModelingToolsEditorOnly/Private/SkeletalMesh/SkeletalMeshEditingInterface.cpp b/Engine/Plugins/Runtime/MeshModelingToolset/Source/MeshModelingToolsEditorOnly/Private/SkeletalMesh/SkeletalMeshEditingInterface.cpp
--- a/Engine/Plugins/Runtime/MeshModelingToolset/Source/MeshModelingToolsEditorOnly/Private/SkeletalMesh/SkeletalMeshEditingInterface.cpp
+++ b/Engine/Plugins/Runtime/MeshModelingToolset/Source/MeshModelingToolsEditorOnly/Private/SkeletalMesh/SkeletalMeshEditingInterface.cpp
@@ -4,31 +4,21 @@
 
 #include UE_INLINE_GENERATED_CPP_BY_NAME(SkeletalMeshEditingInterface)
 
-IMPLEMENT_HIT_PROXY(HBoneHitProxy, HHitProxy)
-
-ISkeletalMeshNotifier& ISkeletalMeshEditingInterface::GetNotifier()
+TSharedPtr<ISkeletalMeshNotifier> ISkeletalMeshEditingInterface::GetNotifier()
 {
 	if (!Notifier)
 	{
-		Notifier.Reset(new FSkeletalMeshToolNotifier(TWeakInterfacePtr<ISkeletalMeshEditingInterface>(this)));
+		Notifier = MakeShared<FSkeletalMeshToolNotifier>(TWeakInterfacePtr<ISkeletalMeshEditingInterface>(this));
 	}
 	
-	return *Notifier;
+	return Notifier;
 }
 
-bool ISkeletalMeshEditingInterface::NeedsNotification() const
-{
-	return Notifier && Notifier->Delegate().IsBound(); 
-}
 
-void ISkeletalMeshEditingInterface::BindTo(TSharedPtr<ISkeletalMeshEditorBinding> InBinding)
-{
-	Binding = InBinding;
-}
 
-void ISkeletalMeshEditingInterface::Unbind()
+bool ISkeletalMeshEditingInterface::NeedsNotification() const
 {
-	Binding.Reset();
+	return Notifier && Notifier->Delegate().IsBound(); 
 }
 
 TWeakObjectPtr<USkeletonModifier> ISkeletalMeshEditingInterface::GetModifier() const
@@ -36,28 +26,6 @@ TWeakObjectPtr<USkeletonModifier> ISkeletalMeshEditingInterface::GetModifier() c
 	return nullptr;
 }
 
-TOptional<FName> ISkeletalMeshEditingInterface::GetBoneName(HHitProxy* InHitProxy) const
-{
-	if (const HBoneHitProxy* BoneProxy = HitProxyCast<HBoneHitProxy>(InHitProxy))
-	{
-		return BoneProxy->BoneName;
-	}
-
-	return Binding.IsValid() && Binding.Pin()->GetNameFunction() ? Binding.Pin()->GetNameFunction()(InHitProxy) : TOptional<FName>();
-}
-
-TArray<FName> ISkeletalMeshEditingInterface::GetSelectedBones() const
-{
-	static const TArray<FName> Dummy;
-	return Binding.IsValid() ? Binding.Pin()->GetSelectedBones() : Dummy;
-}
-
-TArray<FName> ISkeletalMeshEditingInterface::GetSelectedMorphTargets() const
-{
-	static const TArray<FName> Dummy;
-	return Binding.IsValid() ? Binding.Pin()->GetSelectedMorphTargets() : Dummy;
-}
-
 FSkeletalMeshToolNotifier::FSkeletalMeshToolNotifier(TWeakInterfacePtr<ISkeletalMeshEditingInterface> InInterface)
 	: ISkeletalMeshNotifier()
 	, Interface(InInterface)
diff --git a/Engine/Plugins/Runtime/MeshModelingToolset/Source/MeshModelingToolsEditorOnly/Private/SkeletalMesh/SkeletalMeshToolsHelper.cpp b/Engine/Plugins/Runtime/MeshModelingToolset/Source/MeshModelingToolsEditorOnly/Private/SkeletalMesh/SkeletalMeshToolsHelper.cpp
new file mode 100644
--- /dev/null
+++ b/Engine/Plugins/Runtime/MeshModelingToolset/Source/MeshModelingToolsEditorOnly/Private/SkeletalMesh/SkeletalMeshToolsHelper.cpp
@@ -0,0 +1,221 @@
+﻿// Copyright Epic Games, Inc. All Rights Reserved.
+
+
+#include "SkeletalMesh/SkeletalMeshToolsHelper.h"
+
+#include "BoneWeights.h"
+#include "SkeletalMeshAttributes.h"
+#include "Components/BaseDynamicMeshComponent.h"
+#include "DynamicMesh/DynamicVertexSkinWeightsAttribute.h"
+#include "Async/ParallelFor.h"
+
+void SkeletalMeshToolsHelper::GetUnposedMesh(
+	TFunctionRef<void(FVertInfo, const FVector&)> WriteFunc,
+	const FDynamicMesh3& PosedMesh,
+	const FDynamicMesh3& SourceMesh,
+	const TArray<FMatrix>& BoneMatrices,
+	FName SkinWeightProfile,
+	const TMap<FName, float>& MorphTargetWeights,
+	const TArray<int32>& VertArray
+	)
+{
+
+	using namespace UE::Geometry;
+	using namespace UE::AnimationCore;
+
+	SkinWeightProfile = SkinWeightProfile == NAME_None ? FSkeletalMeshAttributes::DefaultSkinWeightProfileName : SkinWeightProfile;
+	FDynamicMeshVertexSkinWeightsAttribute* SkinWeightAttribute = SourceMesh.Attributes()->GetSkinWeightsAttribute(SkinWeightProfile);
+
+	bool bHasVertSelection = VertArray.Num() > 0;
+	int32 NumToProcess = bHasVertSelection ? VertArray.Num() : SourceMesh.MaxVertexID();
+	
+	ParallelFor(NumToProcess, [&](int32 Index)
+		{
+			int32 VertID = bHasVertSelection ? VertArray[Index] : Index;
+			
+			if (!SourceMesh.IsVertex(VertID))
+			{
+				return;
+			}
+				
+			FBoneWeights BoneWeights;
+			SkinWeightAttribute->GetValue(VertID, BoneWeights);
+			
+			FMatrix SkinMatrix = FMatrix(FVector::ZeroVector, FVector::ZeroVector, FVector::ZeroVector, FVector::ZeroVector);
+			SkinMatrix.M[3][3] = 0.0f;
+
+			for (const FBoneWeight& BW : BoneWeights)
+			{
+				SkinMatrix += BoneMatrices[BW.GetBoneIndex()] * BW.GetWeight();
+			}
+			FVector PosedVertPos = PosedMesh.GetVertex(VertID);
+				
+			FVector UnposedVertPos = SkinMatrix.Inverse().TransformPosition(PosedVertPos);
+
+			for (const TPair<FName, float>& MorphTargetWeight : MorphTargetWeights)
+			{
+				if (!FMath::IsNearlyZero(MorphTargetWeight.Value))
+				{
+					const FDynamicMeshMorphTargetAttribute* MorphTargetAttribute =
+						SourceMesh.Attributes()->GetMorphTargetAttribute(MorphTargetWeight.Key);
+
+					FVector Delta;
+					MorphTargetAttribute->GetValue(VertID, Delta);
+				
+					UnposedVertPos += MorphTargetWeight.Value * Delta * (-1.0f);
+				}
+			}
+
+			FVertInfo Info = {
+				bHasVertSelection ? Index : INDEX_NONE,
+				VertID,
+			};
+			
+			WriteFunc(Info, UnposedVertPos);
+		});
+}
+
+void SkeletalMeshToolsHelper::GetPosedMesh(
+	TFunctionRef<void(int32, const FVector&)> WriteFunc, 
+	const FDynamicMesh3& SourceMesh,
+	const TArray<FMatrix>& BoneMatrices,
+	FName SkinWeightProfile,
+	const TMap<FName, float>& MorphTargetWeights
+	)
+{
+	using namespace UE::Geometry;
+	using namespace UE::AnimationCore;
+	SkinWeightProfile = SkinWeightProfile == NAME_None ? FSkeletalMeshAttributes::DefaultSkinWeightProfileName : SkinWeightProfile;	
+	FDynamicMeshVertexSkinWeightsAttribute* SkinWeightAttribute = SourceMesh.Attributes()->GetSkinWeightsAttribute(SkinWeightProfile);	
+	ParallelFor(SourceMesh.MaxVertexID(), [&](int32 VertID)
+		{
+			if (!SourceMesh.IsVertex(VertID))
+			{
+				return;
+			}
+				
+			FVector PosedVertPos = SourceMesh.GetVertex(VertID);
+			for (const TPair<FName, float>& MorphTargetWeight : MorphTargetWeights)
+			{
+				if (!FMath::IsNearlyZero(MorphTargetWeight.Value))
+				{
+					const FDynamicMeshMorphTargetAttribute* MorphTargetAttribute =
+						SourceMesh.Attributes()->GetMorphTargetAttribute(MorphTargetWeight.Key);
+
+					FVector Delta;
+					MorphTargetAttribute->GetValue(VertID, Delta);
+			
+					PosedVertPos += MorphTargetWeight.Value * Delta;
+				}
+			}	
+				
+
+			FBoneWeights BoneWeights;
+			SkinWeightAttribute->GetValue(VertID, BoneWeights);
+			
+			FMatrix SkinMatrix = FMatrix(FVector::ZeroVector, FVector::ZeroVector, FVector::ZeroVector, FVector::ZeroVector);
+			SkinMatrix.M[3][3] = 0.0f;
+
+			for (const FBoneWeight& BW : BoneWeights)
+			{
+				SkinMatrix += BoneMatrices[BW.GetBoneIndex()] * BW.GetWeight();
+			}
+			
+			PosedVertPos = SkinMatrix.TransformPosition(PosedVertPos);
+
+			WriteFunc(VertID, PosedVertPos);
+		});
+}
+
+TArray<FMatrix> SkeletalMeshToolsHelper::ComputeBoneMatrices(const TArray<FTransform>& ComponentSpaceTransformsRefPose,
+	const TArray<FTransform>& ComponentSpaceTransforms)
+{
+	check(ComponentSpaceTransformsRefPose.Num() == ComponentSpaceTransforms.Num());
+	
+	TArray<FMatrix> BoneMatrices;
+	BoneMatrices.SetNumUninitialized(ComponentSpaceTransforms.Num());
+	for (int32 BoneIndex = 0; BoneIndex < ComponentSpaceTransforms.Num(); ++BoneIndex)
+	{
+		BoneMatrices[BoneIndex] =
+				ComponentSpaceTransformsRefPose[BoneIndex].ToMatrixWithScale().Inverse() *
+				ComponentSpaceTransforms[BoneIndex].ToMatrixWithScale();
+	}
+
+	return BoneMatrices;
+}
+
+void SkeletalMeshToolsHelper::FPoseChangeDetector::CheckPose(const TArray<FTransform>& ComponentSpaceTransforms,
+	const TMap<FName, float>& MorphTargetWeights)
+{
+	auto HasPoseChanged = [&]()
+		{
+			if (ComponentSpaceTransforms.Num() != PreviousComponentSpaceTransforms.Num())
+			{
+				return true;
+			}
+				
+			for (int32 BoneIndex=0; BoneIndex< ComponentSpaceTransforms.Num(); ++BoneIndex)
+			{
+				const FTransform& CurrentBoneTransform = ComponentSpaceTransforms[BoneIndex];
+				const FTransform& PrevBoneTransform = PreviousComponentSpaceTransforms[BoneIndex];
+				if (!CurrentBoneTransform.Equals(PrevBoneTransform))
+				{
+					return true;
+				}
+			}
+				
+			if (!MorphTargetWeights.OrderIndependentCompareEqual(PreviousMorphTargetWeights))
+			{
+				return true;
+			}
+				
+			return false;
+		};
+
+	const bool bPoseChanged = HasPoseChanged();
+	bool bNotify = bPoseChanged;
+	if (bPoseChanged)
+	{
+		if (State == PoseStoppedChanging)
+		{
+			State = PoseJustChanged;
+		}
+		else if (State == PoseJustChanged)
+		{
+			State = PoseChanged;
+		}
+	}
+	else
+	{
+		if (State != PoseStoppedChanging)
+		{
+			State = PoseStoppedChanging;
+			bNotify = true;
+		}
+	}
+
+	if (bNotify)
+	{
+		if (PreviousComponentSpaceTransforms.IsEmpty())
+		{
+			PreviousComponentSpaceTransforms = ComponentSpaceTransforms;
+			PreviousMorphTargetWeights = MorphTargetWeights;
+		}
+	
+		FPayload Payload = {
+			State,
+			ComponentSpaceTransforms,
+			MorphTargetWeights,
+			PreviousComponentSpaceTransforms,
+			PreviousMorphTargetWeights
+		};
+	
+		Notifier.Broadcast(Payload);	
+	}
+
+	if (bPoseChanged)
+	{
+		PreviousComponentSpaceTransforms = ComponentSpaceTransforms;
+		PreviousMorphTargetWeights = MorphTargetWeights;	
+	}
+}
diff --git a/Engine/Plugins/Runtime/MeshModelingToolset/Source/MeshModelingToolsEditorOnly/Private/SkeletalMesh/SkeletonEditingTool.cpp b/Engine/Plugins/Runtime/MeshModelingToolset/Source/MeshModelingToolsEditorOnly/Private/SkeletalMesh/SkeletonEditingTool.cpp
--- a/Engine/Plugins/Runtime/MeshModelingToolset/Source/MeshModelingToolsEditorOnly/Private/SkeletalMesh/SkeletonEditingTool.cpp
+++ b/Engine/Plugins/Runtime/MeshModelingToolset/Source/MeshModelingToolsEditorOnly/Private/SkeletalMesh/SkeletonEditingTool.cpp
@@ -31,6 +31,8 @@
 
 #define LOCTEXT_NAMESPACE "USkeletonEditingTool"
 
+IMPLEMENT_HIT_PROXY(HBoneHitProxy, HHitProxy)
+
 namespace SkeletonEditingTool
 {
 	
@@ -277,7 +279,7 @@ void USkeletonEditingTool::SetupWatchers()
 			UpdateGizmo();
 			if (NeedsNotification())
 			{
-				GetNotifier().Notify(InBoneNames, ESkeletalMeshNotifyType::BonesSelected);
+				GetNotifier()->Notify(InBoneNames, ESkeletalMeshNotifyType::BonesSelected);
 			}
 		},
 		this->Selection
@@ -416,7 +418,7 @@ void USkeletonEditingTool::Shutdown(EToolShutdownType ShutdownType)
 		// to force to refresh the tree
 		if (NeedsNotification())
 		{
-			GetNotifier().Notify({}, ESkeletalMeshNotifyType::BonesAdded);
+			GetNotifier()->Notify({}, ESkeletalMeshNotifyType::BonesAdded);
 		}
 	}
 	
@@ -593,8 +595,8 @@ void USkeletonEditingTool::RegisterPasteAction(FInteractiveToolActionSet& InOutA
 
 				if (NeedsNotification())
 				{
-					GetNotifier().Notify(NewBones, ESkeletalMeshNotifyType::HierarchyChanged);
-					GetNotifier().Notify(NewBones, ESkeletalMeshNotifyType::BonesSelected);
+					GetNotifier()->Notify(NewBones, ESkeletalMeshNotifyType::HierarchyChanged);
+					GetNotifier()->Notify(NewBones, ESkeletalMeshNotifyType::BonesSelected);
 				}
 
 				EndChange();
@@ -634,8 +636,8 @@ void USkeletonEditingTool::RegisterDuplicateAction(FInteractiveToolActionSet& In
 
 				if (NeedsNotification())
 				{
-					GetNotifier().Notify(NewBones, ESkeletalMeshNotifyType::HierarchyChanged);
-					GetNotifier().Notify(NewBones, ESkeletalMeshNotifyType::BonesSelected);
+					GetNotifier()->Notify(NewBones, ESkeletalMeshNotifyType::HierarchyChanged);
+					GetNotifier()->Notify(NewBones, ESkeletalMeshNotifyType::BonesSelected);
 				}
 
 				EndChange();
@@ -770,7 +772,7 @@ void USkeletonEditingTool::CreateNewBone()
 	{
 		if (NeedsNotification())
 		{
-			GetNotifier().Notify({BoneName}, ESkeletalMeshNotifyType::BonesAdded);
+			GetNotifier()->Notify({BoneName}, ESkeletalMeshNotifyType::BonesAdded);
 		}
 	
 		Selection = {BoneName};
@@ -787,12 +789,12 @@ void USkeletonEditingTool::MirrorBones()
 	TGuardValue OperationGuard(Operation, EEditingOperation::Mirror);
 	BeginChange();
 
-	const bool bBonesMirrored = Modifier->MirrorBones(GetSelectedBones(), MirroringProperties->Options);
+	const bool bBonesMirrored = Modifier->MirrorBones(GetSelection(), MirroringProperties->Options);
 	if (bBonesMirrored)
 	{
 		if (NeedsNotification())
 		{
-			GetNotifier().Notify({}, ESkeletalMeshNotifyType::HierarchyChanged);
+			GetNotifier()->Notify({}, ESkeletalMeshNotifyType::HierarchyChanged);
 		}
 	
 		EndChange();
@@ -804,7 +806,7 @@ void USkeletonEditingTool::MirrorBones()
 
 void USkeletonEditingTool::RemoveBones()
 {
-	const TArray<FName> BonesToRemove = GetSelectedBones();
+	const TArray<FName> BonesToRemove = GetSelection();
 	
 	TGuardValue OperationGuard(Operation, EEditingOperation::Remove);
 	BeginChange();
@@ -819,7 +821,7 @@ void USkeletonEditingTool::RemoveBones()
 		
 		if (NeedsNotification())
 		{
-			GetNotifier().Notify(BonesToRemove, ESkeletalMeshNotifyType::BonesRemoved);
+			GetNotifier()->Notify(BonesToRemove, ESkeletalMeshNotifyType::BonesRemoved);
 		}
 		
 		EndChange();
@@ -836,14 +838,14 @@ void USkeletonEditingTool::UnParentBones()
 	TGuardValue OperationGuard(Operation, EEditingOperation::Parent);
 	BeginChange();
 
-	const bool bBonesUnParented = Modifier->ParentBones(GetSelectedBones(), Dummy);
+	const bool bBonesUnParented = Modifier->ParentBones(GetSelection(), Dummy);
 	if (bBonesUnParented)
 	{
 		GetToolManager()->DisplayMessage(LOCTEXT("Unparent", "Selected bones have been unparented."), EToolMessageLevel::UserNotification);
 
 		if (NeedsNotification())
 		{
-			GetNotifier().Notify({}, ESkeletalMeshNotifyType::HierarchyChanged);
+			GetNotifier()->Notify({}, ESkeletalMeshNotifyType::HierarchyChanged);
 		}
 		
 		EndChange();
@@ -862,7 +864,7 @@ void USkeletonEditingTool::SnapBoneToComponentSelection(const bool bCreate)
 	}
 
 	const FReferenceSkeleton& RefSkeleton = Modifier->GetReferenceSkeleton();
-	const TArray<FName> Bones = GetSelectedBones();
+	const TArray<FName> Bones = GetSelection();
 
 	if (!bCreate)
 	{
@@ -898,7 +900,7 @@ void USkeletonEditingTool::SnapBoneToComponentSelection(const bool bCreate)
 		{
 			if (NeedsNotification())
 			{
-				GetNotifier().Notify({BoneName}, ESkeletalMeshNotifyType::BonesAdded);
+				GetNotifier()->Notify({BoneName}, ESkeletalMeshNotifyType::BonesAdded);
 			}
 	
 			Selection = {BoneName};
@@ -938,12 +940,12 @@ void USkeletonEditingTool::ParentBones(const FName& InParentName)
 	}
 
 	BeginChange();
-	const bool bBonesParented = Modifier->ParentBones(GetSelectedBones(), {InParentName});
+	const bool bBonesParented = Modifier->ParentBones(GetSelection(), {InParentName});
 	if (bBonesParented)
 	{
 		if (NeedsNotification())
 		{
-			GetNotifier().Notify({}, ESkeletalMeshNotifyType::HierarchyChanged);
+			GetNotifier()->Notify({}, ESkeletalMeshNotifyType::HierarchyChanged);
 		}
 		
 		EndChange();
@@ -955,11 +957,20 @@ void USkeletonEditingTool::ParentBones(const FName& InParentName)
 	Operation = EEditingOperation::Select;
 }
 
+TOptional<FName> USkeletonEditingTool::GetBoneName(HHitProxy* InHitProxy)
+{
+	if (const HBoneHitProxy* BoneProxy = HitProxyCast<HBoneHitProxy>(InHitProxy))
+	{
+		return BoneProxy->BoneName;
+	}
+	return {};
+}
+
 void USkeletonEditingTool::MoveBones()
 {
 	const FReferenceSkeleton& RefSkeleton = Modifier->GetReferenceSkeleton();
 
-	const TArray<FName> Bones = GetSelectedBones();
+	const TArray<FName> Bones = GetSelection();
 	const bool bHasValidBone = Bones.ContainsByPredicate([&](const FName& InBoneName)
 	{
 		return RefSkeleton.FindRawBoneIndex(InBoneName) > INDEX_NONE;
@@ -978,7 +989,7 @@ void USkeletonEditingTool::MoveBones()
 	{
 		// if (NeedsNotification())
 		// {
-		// 	GetNotifier().Notify(Bones, ESkeletalMeshNotifyType::BonesMoved);
+		// 	GetNotifier()->Notify(Bones, ESkeletalMeshNotifyType::BonesMoved);
 		// }
 		
 		EndChange();
@@ -1029,7 +1040,7 @@ void USkeletonEditingTool::RenameBones()
 		
 		if (NeedsNotification())
 		{
-			GetNotifier().Notify( {NewName}, ESkeletalMeshNotifyType::BonesRenamed);
+			GetNotifier()->Notify( {NewName}, ESkeletalMeshNotifyType::BonesRenamed);
 		}
 		
 		EndChange();
@@ -1368,7 +1379,7 @@ void USkeletonEditingTool::OrientBones()
 {
 	const FReferenceSkeleton& RefSkeleton = Modifier->GetReferenceSkeleton();
 
-	const TArray<FName> Bones = GetSelectedBones();
+	const TArray<FName> Bones = GetSelection();
 	const bool bHasValidBone = Bones.ContainsByPredicate([&](const FName& InBoneName)
 	{
 		return RefSkeleton.FindRawBoneIndex(InBoneName) > INDEX_NONE;
@@ -1388,7 +1399,7 @@ void USkeletonEditingTool::OrientBones()
 		UpdateGizmo();
 		// if (NeedsNotification())
 		// {
-		// 	GetNotifier().Notify(Bones, ESkeletalMeshNotifyType::BonesMoved);
+		// 	GetNotifier()->Notify(Bones, ESkeletalMeshNotifyType::BonesMoved);
 		// }
 		
 		EndChange();
@@ -1545,7 +1556,6 @@ void USkeletonEditingTool::OnClickRelease(const FInputDeviceRay& InReleasePos)
 
 void USkeletonEditingTool::OnTerminateDragSequence()
 {}
-
 void USkeletonEditingTool::OnTick(float DeltaTime)
 {
 	const FViewport* Viewport = GetToolManager()->GetContextQueriesAPI()->GetFocusedViewport();
@@ -1676,7 +1686,7 @@ TWeakObjectPtr<USkeletonModifier> USkeletonEditingTool::GetModifier() const
 
 void USkeletonEditingTool::HandleSkeletalMeshModified(const TArray<FName>& InBoneNames, const ESkeletalMeshNotifyType InNotifyType)
 {
-	if (GetNotifier().Notifying())
+	if (GetNotifier()->Notifying())
 	{
 		return;
 	}
@@ -1874,11 +1884,6 @@ FBox USkeletonEditingTool::GetWorldSpaceFocusBox()
 	return USingleSelectionTool::GetWorldSpaceFocusBox();
 }
 
-TArray<FName> USkeletonEditingTool::GetSelectedBones() const
-{
-	return GetSelection();
-}
-
 const TArray<FName>& USkeletonEditingTool::GetSelection() const
 {
 	return Selection;
diff --git a/Engine/Plugins/Runtime/MeshModelingToolset/Source/MeshModelingToolsEditorOnly/Private/SkeletalMesh/SkinWeightsPaintTool.cpp b/Engine/Plugins/Runtime/MeshModelingToolset/Source/MeshModelingToolsEditorOnly/Private/SkeletalMesh/SkinWeightsPaintTool.cpp
--- a/Engine/Plugins/Runtime/MeshModelingToolset/Source/MeshModelingToolsEditorOnly/Private/SkeletalMesh/SkinWeightsPaintTool.cpp
+++ b/Engine/Plugins/Runtime/MeshModelingToolset/Source/MeshModelingToolsEditorOnly/Private/SkeletalMesh/SkinWeightsPaintTool.cpp
@@ -57,6 +57,7 @@
 #include "TargetInterfaces/DynamicMeshProvider.h"
 #include "TargetInterfaces/SkeletalMeshBackedTarget.h"
 #include "TargetInterfaces/SkeletonProvider.h"
+#include "ToolTargets/DynamicMeshComponentToolTarget.h"
 
 #include UE_INLINE_GENERATED_CPP_BY_NAME(SkinWeightsPaintTool)
 
@@ -317,6 +318,14 @@ TArray<FName> GetProfilesFromToolTarget(UToolTarget* InTarget, const FName InLOD
 	if (InTarget)
 	{
 		const EMeshLODIdentifier LODId = GetLODId(InLOD);
+
+		// offer a fast path if target is a dynamic mesh component to avoid conversion to mesh description
+		if (UDynamicMeshComponentToolTarget* DynamicMeshToolTarget = Cast<UDynamicMeshComponentToolTarget>(InTarget))
+		{
+			TArray<FName> Profiles;
+			DynamicMeshToolTarget->GetDynamicMesh().Attributes()->GetSkinWeightsAttributes().GenerateKeyArray(Profiles);
+			return Profiles;	
+		}
 		
 		if (IMeshDescriptionProvider* MeshDescriptionProvider = Cast<IMeshDescriptionProvider>(InTarget))
 		{
diff --git a/Engine/Plugins/Runtime/MeshModelingToolset/Source/MeshModelingToolsEditorOnly/Public/SkeletalMesh/RefSkeletonPoser.h b/Engine/Plugins/Runtime/MeshModelingToolset/Source/MeshModelingToolsEditorOnly/Public/SkeletalMesh/RefSkeletonPoser.h
new file mode 100644
--- /dev/null
+++ b/Engine/Plugins/Runtime/MeshModelingToolset/Source/MeshModelingToolsEditorOnly/Public/SkeletalMesh/RefSkeletonPoser.h
@@ -0,0 +1,74 @@
+﻿// Copyright Epic Games, Inc. All Rights Reserved.
+
+#pragma once
+
+#include "InteractiveToolChange.h"
+#include "UObject/Object.h"
+#include "ReferenceSkeleton.h"
+
+#include "RefSkeletonPoser.generated.h"
+
+#define UE_API MESHMODELINGTOOLSEDITORONLY_API
+
+
+UCLASS(MinimalAPI)
+class URefSkeletonPoser: public UObject
+{
+	GENERATED_BODY()
+
+public:
+	UE_API void SetRefSkeleton(const FReferenceSkeleton& InRefSkeleton);
+
+	UE_API void ModifyBoneAdditiveTransform(int32 BoneIndex, TFunctionRef<void(FTransform&)> InModifyFunc);
+	UE_API void ClearBoneAdditiveTransform(int32 BoneIndex);
+	UE_API void ClearAllBoneAdditiveTransforms();
+
+	UE_API TOptional<FTransform> GetBoneAdditiveTransform(int32 BoneIndex);
+
+	UE_API const TArray<FTransform>& GetComponentSpaceTransforms();
+
+	UE_API const FTransform& GetComponentSpaceTransform(const uint32 Index);
+
+	UE_API const TArray<FTransform>& GetComponentSpaceTransformsRefPose();
+
+	UE_API void BeginPoseChange();
+	UE_API void EndPoseChange();
+	UE_API bool IsRecordingPoseChange();
+
+protected:
+	void Invalidate(const uint32 Index = INDEX_NONE);
+	
+	FReferenceSkeleton RefSkeleton;
+	TArray<FTransform> Transforms;
+	TBitArray<> TransformCached;
+
+	TMap<int32, FTransform> AdditiveTransforms;
+		
+	TArray<FTransform> RefTransforms;
+
+	class FPoseChange : public FToolCommandChange
+	{
+	public:
+		FPoseChange(URefSkeletonPoser* InPoser);
+	
+		TMap<int32, FTransform> OldAdditiveTransforms;
+		TMap<int32, FTransform> NewAdditiveTransforms;
+	
+		virtual FString ToString() const override;
+
+		virtual void Apply(UObject* Object) override;
+
+		virtual void Revert(UObject* Object) override;
+
+		void Close(URefSkeletonPoser* InPoser);
+	};
+	
+	TUniquePtr<FPoseChange> ActivePoseChange;
+
+	bool bPoseChanged = false;
+};
+
+
+
+
+#undef UE_API
diff --git a/Engine/Plugins/Runtime/MeshModelingToolset/Source/MeshModelingToolsEditorOnly/Public/SkeletalMesh/SkeletalMeshEditingInterface.h b/Engine/Plugins/Runtime/MeshModelingToolset/Source/MeshModelingToolsEditorOnly/Public/SkeletalMesh/SkeletalMeshEditingInterface.h
--- a/Engine/Plugins/Runtime/MeshModelingToolset/Source/MeshModelingToolsEditorOnly/Public/SkeletalMesh/SkeletalMeshEditingInterface.h
+++ b/Engine/Plugins/Runtime/MeshModelingToolset/Source/MeshModelingToolsEditorOnly/Public/SkeletalMesh/SkeletalMeshEditingInterface.h
@@ -42,25 +42,16 @@ class ISkeletalMeshEditingInterface
 	GENERATED_BODY()
 
 public:
-	UE_API ISkeletalMeshNotifier& GetNotifier();
+	UE_API TSharedPtr<ISkeletalMeshNotifier> GetNotifier();
 	UE_API bool NeedsNotification() const;
 
-	UE_API virtual TArray<FName> GetSelectedBones() const;
-	UE_API virtual TArray<FName> GetSelectedMorphTargets() const;
-	
-	UE_API void BindTo(TSharedPtr<ISkeletalMeshEditorBinding> InBinding);
-	UE_API void Unbind();
-
 	UE_API virtual TWeakObjectPtr<USkeletonModifier> GetModifier() const;
-	
+
 protected:
 	virtual void HandleSkeletalMeshModified(const TArray<FName>& BoneNames, const ESkeletalMeshNotifyType InNotifyType) = 0;
-	UE_API TOptional<FName> GetBoneName(HHitProxy* InHitProxy) const;
 	
-	TWeakPtr<ISkeletalMeshEditorBinding> Binding;
-
 private:
-	TUniquePtr<FSkeletalMeshToolNotifier> Notifier;
+	TSharedPtr<FSkeletalMeshToolNotifier> Notifier;
 	
 	friend FSkeletalMeshToolNotifier;
 };
@@ -159,6 +150,11 @@ public:
 	virtual void ShowSkeleton() PURE_VIRTUAL(USkeletalMeshEditorContextObjectBase::ShowSkeleton, return;)
 	virtual void ToggleBoneManipulation(bool bEnable) PURE_VIRTUAL(USkeletalMeshEditorContextObjectBase::ToggleBoneManipulation, return;)
 	virtual const TArray<FTransform>& GetComponentSpaceBoneTransforms(UToolTarget* InToolTarget) PURE_VIRTUAL(USkeletalMeshEditorContextObjectBase::GetComponentSpaceBoneTransforms, static TArray<FTransform> Dummy; return Dummy;)
+
+	virtual FName GetEditingMorphTarget()  PURE_VIRTUAL(USkeletalMeshEditorContextObjectBase::GetEditingMorphTarget, return NAME_None;);
+	virtual TMap<FName, float> GetMorphTargetWeights() PURE_VIRTUAL(USkeletalMeshEditorContextObjectBase::GetMorphTargetWeights, return {};);	
+	virtual void NotifyMorphTargetEdited() PURE_VIRTUAL(USkeletalMeshEditorContextObjectBase::NotifyMorphTargetEdited, return;);
+	
 	virtual void BindTo(ISkeletalMeshEditingInterface* InEditingInterface) PURE_VIRTUAL(USkeletalMeshEditorContextObjectBase::BindTo, return;)
 	virtual void UnbindFrom(ISkeletalMeshEditingInterface* InEditingInterface) PURE_VIRTUAL(USkeletalMeshEditorContextObjectBase::UnbindFrom, return;)
 };
diff --git a/Engine/Plugins/Runtime/MeshModelingToolset/Source/MeshModelingToolsEditorOnly/Public/SkeletalMesh/SkeletalMeshToolsHelper.h b/Engine/Plugins/Runtime/MeshModelingToolset/Source/MeshModelingToolsEditorOnly/Public/SkeletalMesh/SkeletalMeshToolsHelper.h
new file mode 100644
--- /dev/null
+++ b/Engine/Plugins/Runtime/MeshModelingToolset/Source/MeshModelingToolsEditorOnly/Public/SkeletalMesh/SkeletalMeshToolsHelper.h
@@ -0,0 +1,82 @@
+﻿// Copyright Epic Games, Inc. All Rights Reserved.
+
+#pragma once
+
+#include "CoreMinimal.h"
+#include "DynamicSubmesh3.h"
+
+#define UE_API MESHMODELINGTOOLSEDITORONLY_API
+
+namespace UE::Geometry { class FDynamicMesh3; }
+
+namespace SkeletalMeshToolsHelper
+{
+	using namespace UE::Geometry;
+
+	struct FVertInfo
+	{
+		int32 VertArrayIndex;
+		int32 VertID;
+	};
+
+	// Writes out the Mesh by first resetting the skeleton pose and then disable the morph targets
+	// it can optionally take a Vert Array to only process a subset of verts
+	UE_API void GetUnposedMesh(
+		TFunctionRef<void(FVertInfo, const FVector&)> WriteFunc,
+		const FDynamicMesh3& PosedMesh,
+		const FDynamicMesh3& SourceMesh,
+		const TArray<FMatrix>& BoneMatrices,
+		FName SkinWeightProfile,
+		const TMap<FName, float>& MorphTargetWeights,
+		const TArray<int32>& VertArray = {}
+		);
+
+	
+	UE_API void GetPosedMesh(
+		TFunctionRef<void(int32, const FVector&)> WriteFunc,
+		const FDynamicMesh3& SourceMesh,
+		const TArray<FMatrix>& BoneMatrices,
+		FName SkinWeightProfile,
+		const TMap<FName, float>& MorphTargetWeights
+		);
+	
+	UE_API TArray<FMatrix> ComputeBoneMatrices(
+		const TArray<FTransform>& ComponentSpaceTransformsRefPose,
+		const TArray<FTransform>& ComponentSpaceTransforms
+		);
+
+	struct FPoseChangeDetector
+	{
+		enum EState
+		{
+			PoseJustChanged,
+			PoseChanged,
+			PoseStoppedChanging
+		};
+		
+		struct FPayload
+		{
+			EState CurrentState;
+			const TArray<FTransform>& ComponentSpaceTransforms;
+			const TMap<FName, float>& MorphTargetWeights;
+			const TArray<FTransform>& PreviousComponentSpaceTransforms;
+			const TMap<FName, float>& PreviousMorphTargetWeights;
+		};
+
+		DECLARE_MULTICAST_DELEGATE_OneParam(FNotifier, FPayload);	
+		
+		UE_API void CheckPose(const TArray<FTransform>& ComponentSpaceTransforms, const TMap<FName,float>& MorphTargetWeights);
+		UE_API FNotifier& GetNotifier() {return Notifier;}
+	
+	protected:
+		FNotifier Notifier;
+	
+		TArray<FTransform> PreviousComponentSpaceTransforms;
+		TMap<FName, float> PreviousMorphTargetWeights;
+
+		EState State = PoseStoppedChanging;
+	};
+
+}
+
+#undef UE_API
diff --git a/Engine/Plugins/Runtime/MeshModelingToolset/Source/MeshModelingToolsEditorOnly/Public/SkeletalMesh/SkeletonEditingTool.h b/Engine/Plugins/Runtime/MeshModelingToolset/Source/MeshModelingToolsEditorOnly/Public/SkeletalMesh/SkeletonEditingTool.h
--- a/Engine/Plugins/Runtime/MeshModelingToolset/Source/MeshModelingToolsEditorOnly/Public/SkeletalMesh/SkeletonEditingTool.h
+++ b/Engine/Plugins/Runtime/MeshModelingToolset/Source/MeshModelingToolsEditorOnly/Public/SkeletalMesh/SkeletonEditingTool.h
@@ -146,9 +146,6 @@ public:
 	UE_API void UnParentBones();
 	UE_API void SnapBoneToComponentSelection(const bool bCreate);
 	
-	// ISkeletalMeshEditingInterface overrides
-	UE_API virtual TArray<FName> GetSelectedBones() const override;
-
 	UE_API const TArray<FName>& GetSelection() const;
 
 	UE_API const FTransform& GetTransform(const FName InBoneName, const bool bWorld) const;
@@ -177,6 +174,8 @@ protected:
 	UE_API void CreateNewBone();
 	UE_API void ParentBones(const FName& InParentName);
 
+	UE_API static TOptional<FName> GetBoneName(HHitProxy* InHitProxy);
+	
 public:
 	
 	UPROPERTY()
diff --git a/Engine/Source/Developer/SkeletalMeshUtilitiesCommon/Public/SkeletalMeshNotifier.h b/Engine/Source/Developer/SkeletalMeshUtilitiesCommon/Public/SkeletalMeshNotifier.h
--- a/Engine/Source/Developer/SkeletalMeshUtilitiesCommon/Public/SkeletalMeshNotifier.h
+++ b/Engine/Source/Developer/SkeletalMeshUtilitiesCommon/Public/SkeletalMeshNotifier.h
@@ -49,15 +49,12 @@ public:
 	ISkeletalMeshEditorBinding() = default;
 	virtual ~ISkeletalMeshEditorBinding() = default;
 
-	virtual ISkeletalMeshNotifier& GetNotifier() = 0;
+	virtual TSharedPtr<ISkeletalMeshNotifier> GetNotifier() = 0;
 
-	using NameFunction = TFunction< TOptional<FName>(HHitProxy*) >;
+	using NameFunction = TFunction< TOptional<FName>(HHitProxy*)>;
 	virtual NameFunction GetNameFunction() = 0;
 
 	virtual TArray<FName> GetSelectedBones() const = 0;
-	virtual TArray<FName> GetSelectedMorphTargets() const = 0;
-	
-
 };
 
 #undef UE_API
diff --git a/Engine/Source/Editor/Persona/Private/SMorphTargetViewer.cpp b/Engine/Source/Editor/Persona/Private/SMorphTargetViewer.cpp
--- a/Engine/Source/Editor/Persona/Private/SMorphTargetViewer.cpp
+++ b/Engine/Source/Editor/Persona/Private/SMorphTargetViewer.cpp
@@ -411,9 +411,16 @@ float SMorphTargetListRow::GetWeight() const
 	}
 	else
 	{
-		USkeletalMeshComponent* SkelComp = PreviewScenePtr.Pin()->GetPreviewMeshComponent();
-		return SkelComp->GetMorphTargetCurves()[Item->Name];
+		if (USkeletalMeshComponent* SkelComp = PreviewScenePtr.Pin()->GetPreviewMeshComponent())
+		{
+			if (const float* Weight = SkelComp->GetMorphTargetCurves().Find(Item->Name))
+			{
+				return *Weight;
+			}
+		}
 	}
+
+	return 0;
 }
 //////////////////////////////////////////////////////////////////////////
 // SMorphTargetViewer
diff --git a/Engine/Source/Editor/SkeletalMeshEditor/Private/SkeletalMeshEditor.cpp b/Engine/Source/Editor/SkeletalMeshEditor/Private/SkeletalMeshEditor.cpp
--- a/Engine/Source/Editor/SkeletalMeshEditor/Private/SkeletalMeshEditor.cpp
+++ b/Engine/Source/Editor/SkeletalMeshEditor/Private/SkeletalMeshEditor.cpp
@@ -67,6 +67,7 @@
 #include "Styling/AppStyle.h"
 #include "InterchangeAssetImportData.h"
 #include "InterchangeManager.h"
+#include "SkeletalMeshEditorModeUILayer.h"
 
 const FName SkeletalMeshEditorAppIdentifier = FName(TEXT("SkeletalMeshEditorApp"));
 
@@ -215,9 +216,23 @@ bool FSkeletalMeshEditor::OnRequestClose(EAssetEditorCloseReason InCloseReason)
 	return bAllowClose;
 }
 
+void FSkeletalMeshEditor::SaveAsset_Execute()
+{
+	OnPreSaveAssetDelegate.Broadcast();
+	ISkeletalMeshEditor::SaveAsset_Execute();
+}
+
+void FSkeletalMeshEditor::SaveAssetAs_Execute()
+{
+	OnPreSaveAssetAsDelegate.Broadcast();
+	ISkeletalMeshEditor::SaveAssetAs_Execute();
+}
+
 void FSkeletalMeshEditor::RegisterTabSpawners(const TSharedRef<class FTabManager>& InTabManager)
 {
 	FAssetEditorToolkit::RegisterTabSpawners(InTabManager);
+
+	SkeletalMeshEditorMenuCategory = InTabManager->AddLocalWorkspaceMenuCategory(LOCTEXT("WorkspaceMenu_SkeletalMeshEditor", "SkeletalMeshEditor"));
 }
 
 void FSkeletalMeshEditor::UnregisterTabSpawners(const TSharedRef<class FTabManager>& InTabManager)
@@ -254,7 +269,12 @@ void FSkeletalMeshEditor::InitSkeletalMeshEditor(const EToolkitMode::Type Mode,
 	const bool bCreateDefaultStandaloneMenu = true;
 	const bool bCreateDefaultToolbar = true;
 	FAssetEditorToolkit::InitAssetEditor(Mode, InitToolkitHost, SkeletalMeshEditorAppIdentifier, FTabManager::FLayout::NullLayout, bCreateDefaultStandaloneMenu, bCreateDefaultToolbar, InSkeletalMesh);
-
+	
+	TSharedPtr<class IToolkitHost> PinnedToolkitHost = ToolkitHost.Pin();
+	check(PinnedToolkitHost.IsValid());
+	ModeUILayer = MakeShareable(new FSkeletalMeshEditorModeUILayer(PinnedToolkitHost.Get()));
+	ModeUILayer->SetModeMenuCategory(SkeletalMeshEditorMenuCategory);
+	
 	BindCommands();
 
 	AddApplicationMode(
@@ -280,6 +300,27 @@ void FSkeletalMeshEditor::InitSkeletalMeshEditor(const EToolkitMode::Type Mode,
 	}
 }
 
+void FSkeletalMeshEditor::OnToolkitHostingStarted(const TSharedRef<IToolkit>& Toolkit)
+{
+	ISkeletalMeshEditor::OnToolkitHostingStarted(Toolkit);
+
+	ModeUILayer->OnToolkitHostingStarted(Toolkit);
+}
+
+void FSkeletalMeshEditor::OnToolkitHostingFinished(const TSharedRef<IToolkit>& Toolkit)
+{
+	ISkeletalMeshEditor::OnToolkitHostingFinished(Toolkit);
+	
+	ModeUILayer->OnToolkitHostingFinished(Toolkit);
+}
+
+void FSkeletalMeshEditor::PostInitAssetEditor()
+{
+	ISkeletalMeshEditor::PostInitAssetEditor();
+
+	
+}
+
 FName FSkeletalMeshEditor::GetToolkitFName() const
 {
 	return FName("SkeletalMeshEditor");
@@ -1246,7 +1287,7 @@ void FSkeletalMeshEditor::HandleSelectionChanged(const TArrayView<TSharedPtr<ISk
 					return bIsBoneType && bIsNotNone;
 				},
 				[](const TSharedPtr<ISkeletonTreeItem>& InItem){ return InItem->GetRowItemName(); });
-			Binding->GetNotifier().Notify(BoneSelection, ESkeletalMeshNotifyType::BonesSelected);
+			Binding->GetNotifier()->Notify(BoneSelection, ESkeletalMeshNotifyType::BonesSelected);
 		}
 	}
 }
@@ -1546,13 +1587,38 @@ TSharedPtr<ISkeletalMeshEditorBinding> FSkeletalMeshEditor::GetBinding()
 	return Binding;
 }
 
+FSimpleMulticastDelegate& FSkeletalMeshEditor::OnPreSaveAsset()
+{
+	return OnPreSaveAssetDelegate;
+}
+
+FSimpleMulticastDelegate& FSkeletalMeshEditor::OnPreSaveAssetAs()
+{
+	return OnPreSaveAssetAsDelegate;
+}
+
+TSharedPtr<FUICommandInfo> FSkeletalMeshEditor::GetResetBoneTransformsCommand()
+{
+	const FSkeletalMeshEditorCommands Commands = FSkeletalMeshEditorCommands::Get();
+	return Commands.ResetBoneTransforms;
+}
+
+TSharedPtr<FUICommandInfo> FSkeletalMeshEditor::GetResetAllBonesTransformsCommand()
+{
+	const FSkeletalMeshEditorCommands Commands = FSkeletalMeshEditorCommands::Get();
+	return Commands.ResetAllBonesTransforms;
+}
+
 FSkeletalMeshEditorBinding::FSkeletalMeshEditorBinding(TSharedRef<FSkeletalMeshEditor> InEditor)
 	: Editor(InEditor)
-	, Notifier(InEditor)
 {}
 
-ISkeletalMeshNotifier& FSkeletalMeshEditorBinding::GetNotifier()
+TSharedPtr<ISkeletalMeshNotifier> FSkeletalMeshEditorBinding::GetNotifier()
 {
+	if (!Notifier)
+	{
+		Notifier = MakeShared<FSkeletalMeshEditorNotifier>(Editor.Pin().ToSharedRef());
+	}
 	return Notifier;	
 }
 
@@ -1588,18 +1654,6 @@ TArray<FName> FSkeletalMeshEditorBinding::GetSelectedBones() const
 	return Selection;
 }
 
-TArray<FName> FSkeletalMeshEditorBinding::GetSelectedMorphTargets() const
-{
-	TArray<FName> Selection;
-
-	if (!Editor.IsValid())
-	{
-		return Selection;
-	}
-
-	return Editor.Pin()->GetMorphTargetViewer()->GetSelectedMorphTargetNames();
-}
-
 FSkeletalMeshEditorNotifier::FSkeletalMeshEditorNotifier(TSharedRef<FSkeletalMeshEditor> InEditor)
 	: ISkeletalMeshNotifier()
 	, Editor(InEditor)
diff --git a/Engine/Source/Editor/SkeletalMeshEditor/Private/SkeletalMeshEditor.h b/Engine/Source/Editor/SkeletalMeshEditor/Private/SkeletalMeshEditor.h
--- a/Engine/Source/Editor/SkeletalMeshEditor/Private/SkeletalMeshEditor.h
+++ b/Engine/Source/Editor/SkeletalMeshEditor/Private/SkeletalMeshEditor.h
@@ -13,6 +13,7 @@
 #include "Containers/ArrayView.h"
 #include "Async/Future.h"
 
+class FSkeletalMeshEditorModeUILayer;
 class IDetailLayoutBuilder;
 class IDetailsView;
 class IPersonaToolkit;
@@ -60,6 +61,11 @@ public:
 	/** Edits the specified Skeleton object */
 	void InitSkeletalMeshEditor(const EToolkitMode::Type Mode, const TSharedPtr<class IToolkitHost>& InitToolkitHost, class USkeletalMesh* InSkeletalMesh);
 
+	// FPersonaAssetEditorToolkit
+	virtual void OnToolkitHostingStarted(const TSharedRef<IToolkit>& Toolkit) override;
+	virtual void OnToolkitHostingFinished(const TSharedRef<IToolkit>& Toolkit) override;
+	virtual void PostInitAssetEditor() override;
+
 	/** IHasPersonaToolkit interface */
 	virtual TSharedRef<class IPersonaToolkit> GetPersonaToolkit() const override { return PersonaToolkit.ToSharedRef(); }
 
@@ -80,6 +86,8 @@ public:
 	
 	//~ Begin FAssetEditorToolkit Interface.
 	virtual bool OnRequestClose(EAssetEditorCloseReason InCloseReason) override;
+	virtual void SaveAsset_Execute() override;
+	virtual void SaveAssetAs_Execute() override;
 	//~ End FAssetEditorToolkit Interface.
 
 	/** FEditorUndoClient interface */
@@ -124,6 +132,10 @@ public:
 	TSharedPtr<IToolkit> GetHostedToolkit() const { return HostedToolkit; }
 	
 	virtual TSharedPtr<ISkeletalMeshEditorBinding> GetBinding() override;
+	virtual FSimpleMulticastDelegate& OnPreSaveAsset() override;
+	virtual FSimpleMulticastDelegate& OnPreSaveAssetAs() override;
+	virtual TSharedPtr<FUICommandInfo> GetResetBoneTransformsCommand() override;
+	virtual TSharedPtr<FUICommandInfo> GetResetAllBonesTransformsCommand() override;
 
 private:
 	void HandleObjectSelected(UObject* InObject);
@@ -224,11 +236,19 @@ private:
 	// The toolkit we're currently hosting.
 	TSharedPtr<IToolkit> HostedToolkit;
 
+	// This is set in ModeUILayer to be the menu category where new tabs are registered to be enabled by the user
+	TSharedPtr<FWorkspaceItem> SkeletalMeshEditorMenuCategory;
+	
+	TSharedPtr<FSkeletalMeshEditorModeUILayer> ModeUILayer;
+
 	// The toolbox widget
 	TSharedPtr<SSkeletalMeshEditorToolbox> ToolboxWidget;
 
 	// Binding to send/receive skeletal mesh modifications
 	TSharedPtr<FSkeletalMeshEditorBinding> Binding;
+
+	FSimpleMulticastDelegate OnPreSaveAssetDelegate;
+	FSimpleMulticastDelegate OnPreSaveAssetAsDelegate;
 };
 
 /**
@@ -254,13 +274,11 @@ class FSkeletalMeshEditorBinding: public ISkeletalMeshEditorBinding
 public:
 	FSkeletalMeshEditorBinding(TSharedRef<FSkeletalMeshEditor> InEditor);
 
-	virtual ISkeletalMeshNotifier& GetNotifier() override;
+	virtual TSharedPtr<ISkeletalMeshNotifier> GetNotifier() override;
 	virtual NameFunction GetNameFunction() override;
 	virtual TArray<FName> GetSelectedBones() const override;
-	virtual TArray<FName> GetSelectedMorphTargets() const override;
-
 
 private:
 	TWeakPtr<FSkeletalMeshEditor> Editor;
-	FSkeletalMeshEditorNotifier Notifier;
+	TSharedPtr<FSkeletalMeshEditorNotifier> Notifier;
 };
diff --git a/Engine/Source/Editor/SkeletalMeshEditor/Private/SkeletalMeshEditorMode.cpp b/Engine/Source/Editor/SkeletalMeshEditor/Private/SkeletalMeshEditorMode.cpp
--- a/Engine/Source/Editor/SkeletalMeshEditor/Private/SkeletalMeshEditorMode.cpp
+++ b/Engine/Source/Editor/SkeletalMeshEditor/Private/SkeletalMeshEditorMode.cpp
@@ -12,6 +12,7 @@
 #include "Styling/AppStyle.h"
 #include "Engine/SkeletalMesh.h"
 #include "AnimAssetFindReplace.h"
+#include "ISkeletalMeshEditorModule.h"
 #include "PersonaTabs.h"
 
 #define LOCTEXT_NAMESPACE "SkeletalMeshEditorMode"
@@ -119,6 +120,10 @@ FSkeletalMeshEditorMode::FSkeletalMeshEditorMode(
 	PersonaModule.OnRegisterTabs().Broadcast(TabFactories, InHostingApp);
 	LayoutExtender = MakeShared<FLayoutExtender>();
 	PersonaModule.OnRegisterLayoutExtensions().Broadcast(*LayoutExtender.Get());
+
+	ISkeletalMeshEditorModule* Module = &FModuleManager::LoadModuleChecked<ISkeletalMeshEditorModule>("SkeletalMeshEditor");
+	Module->OnRegisterLayoutExtensions().Broadcast(*LayoutExtender);
+	
 	TabLayout->ProcessExtensions(*LayoutExtender.Get());
 
 	SkeletalMeshEditor->RegisterModeToolbarIfUnregistered(GetModeName());
diff --git a/Engine/Source/Editor/SkeletalMeshEditor/Private/SkeletalMeshEditorModeUILayer.cpp b/Engine/Source/Editor/SkeletalMeshEditor/Private/SkeletalMeshEditorModeUILayer.cpp
new file mode 100644
--- /dev/null
+++ b/Engine/Source/Editor/SkeletalMeshEditor/Private/SkeletalMeshEditorModeUILayer.cpp
@@ -0,0 +1,56 @@
+// Copyright Epic Games, Inc. All Rights Reserved.
+
+#include "SkeletalMeshEditorModeUILayer.h"
+
+#include "ISkeletalMeshEditorModule.h"
+#include "SkeletalMeshEditor.h"
+#include "WorkspaceMenuStructure.h"
+#include "WorkspaceMenuStructureModule.h"
+#include "Modules/ModuleManager.h"
+
+void USkeletalMeshEditorUISubsystem::Initialize(FSubsystemCollectionBase& Collection)
+{
+	ISkeletalMeshEditorModule& SkeletalMeshEditorModule = FModuleManager::GetModuleChecked<ISkeletalMeshEditorModule>("SkeletalMeshEditor");
+	SkeletalMeshEditorModule.OnRegisterLayoutExtensions().AddUObject(this, &USkeletalMeshEditorUISubsystem::RegisterLayoutExtensions);
+}
+
+void USkeletalMeshEditorUISubsystem::Deinitialize()
+{
+	ISkeletalMeshEditorModule& SkeletalMeshEditorModule = FModuleManager::GetModuleChecked<ISkeletalMeshEditorModule>("SkeletalMeshEditor");
+	SkeletalMeshEditorModule.OnRegisterLayoutExtensions().RemoveAll(this);
+}
+
+void USkeletalMeshEditorUISubsystem::RegisterLayoutExtensions(FLayoutExtender& Extender)
+{
+}
+
+void FSkeletalMeshEditorModeUILayer::OnToolkitHostingStarted(const TSharedRef<IToolkit>& Toolkit)
+{
+	if (!Toolkit->IsAssetEditor())		// We only want to host Mode toolkits
+	{
+		FAssetEditorModeUILayer::OnToolkitHostingStarted(Toolkit);
+		HostedToolkit = Toolkit;
+		Toolkit->SetModeUILayer(SharedThis(this));
+		Toolkit->RegisterTabSpawners(ToolkitHost->GetTabManager().ToSharedRef());
+		RegisterModeTabSpawners();
+		OnToolkitHostReadyForUI.ExecuteIfBound();
+	}
+}
+
+void FSkeletalMeshEditorModeUILayer::OnToolkitHostingFinished(const TSharedRef<IToolkit>& Toolkit)
+{
+	if (HostedToolkit.IsValid() && HostedToolkit.Pin() == Toolkit)	// don't execute OnToolkitHostShutdownUI if the input Toolkit isn't the one we're hosting
+	{
+		FAssetEditorModeUILayer::OnToolkitHostingFinished(Toolkit);
+	}	
+}
+
+TSharedPtr<FWorkspaceItem> FSkeletalMeshEditorModeUILayer::GetModeMenuCategory() const
+{
+	return SkeletalMeshEditorMenuCategory;
+}
+
+void FSkeletalMeshEditorModeUILayer::SetModeMenuCategory(TSharedPtr<FWorkspaceItem> InMenuCategory)
+{
+	SkeletalMeshEditorMenuCategory = InMenuCategory;
+}
diff --git a/Engine/Source/Editor/SkeletalMeshEditor/Private/SkeletalMeshEditorModeUILayer.h b/Engine/Source/Editor/SkeletalMeshEditor/Private/SkeletalMeshEditorModeUILayer.h
new file mode 100644
--- /dev/null
+++ b/Engine/Source/Editor/SkeletalMeshEditor/Private/SkeletalMeshEditorModeUILayer.h
@@ -0,0 +1,42 @@
+// Copyright Epic Games, Inc. All Rights Reserved.
+
+#pragma once
+
+#include "Toolkits/AssetEditorModeUILayer.h"
+
+#include "SkeletalMeshEditorModeUILayer.generated.h"
+
+#define UE_API SKELETALMESHEDITOR_API
+
+/** Interchange layer to manage built in tab locations within the editor's layout. **/
+UCLASS(MinimalAPI)
+class USkeletalMeshEditorUISubsystem : public UAssetEditorUISubsystem
+{
+	GENERATED_BODY()
+public:
+
+	UE_API virtual void Initialize(FSubsystemCollectionBase& Collection) override;
+	UE_API virtual void Deinitialize() override;
+	UE_API virtual void RegisterLayoutExtensions(FLayoutExtender& Extender) override;
+
+};
+
+/** Handles the hosting of additional toolkits, such as the mode toolkit, within the Skeletal Mesh Editor 's toolkit. **/
+class FSkeletalMeshEditorModeUILayer : public FAssetEditorModeUILayer
+{
+public:
+	FSkeletalMeshEditorModeUILayer(const IToolkitHost* InToolkitHost) : FAssetEditorModeUILayer(InToolkitHost) {}
+	virtual void OnToolkitHostingStarted(const TSharedRef<IToolkit>& Toolkit) override;
+	virtual void OnToolkitHostingFinished(const TSharedRef<IToolkit>& Toolkit) override;
+
+	virtual TSharedPtr<FWorkspaceItem> GetModeMenuCategory() const override;
+
+	void SetModeMenuCategory(TSharedPtr<FWorkspaceItem> InMenuCategory);
+protected:
+	
+	TSharedPtr<FWorkspaceItem> SkeletalMeshEditorMenuCategory;
+};
+
+
+
+#undef UE_API
diff --git a/Engine/Source/Editor/SkeletalMeshEditor/Private/SkeletalMeshEditorModule.cpp b/Engine/Source/Editor/SkeletalMeshEditor/Private/SkeletalMeshEditorModule.cpp
--- a/Engine/Source/Editor/SkeletalMeshEditor/Private/SkeletalMeshEditorModule.cpp
+++ b/Engine/Source/Editor/SkeletalMeshEditor/Private/SkeletalMeshEditorModule.cpp
@@ -43,12 +43,16 @@ public:
 	virtual TSharedPtr<FExtensibilityManager> GetMenuExtensibilityManager() override { return MenuExtensibilityManager; }
 	virtual TSharedPtr<FExtensibilityManager> GetToolBarExtensibilityManager() override { return ToolBarExtensibilityManager; }
 
+	virtual FOnRegisterLayoutExtensions& OnRegisterLayoutExtensions() override { return RegisterLayoutExtensions; }
+	
 private:
 	TSharedPtr<FExtensibilityManager> MenuExtensibilityManager;
 	TSharedPtr<FExtensibilityManager> ToolBarExtensibilityManager;
 
 	TArray<FOnSkeletalMeshEditorInitialized> SkeletalMeshEditorInitialized;
 	TArray<FSkeletalMeshEditorToolbarExtender> SkeletalMeshEditorToolbarExtenders;
+	
+	FOnRegisterLayoutExtensions	RegisterLayoutExtensions;
 };
 
 IMPLEMENT_MODULE(FSkeletalMeshEditorModule, SkeletalMeshEditor);
diff --git a/Engine/Source/Editor/SkeletalMeshEditor/Private/SkeletalMeshEditorSubsystem.cpp b/Engine/Source/Editor/SkeletalMeshEditor/Private/SkeletalMeshEditorSubsystem.cpp
--- a/Engine/Source/Editor/SkeletalMeshEditor/Private/SkeletalMeshEditorSubsystem.cpp
+++ b/Engine/Source/Editor/SkeletalMeshEditor/Private/SkeletalMeshEditorSubsystem.cpp
@@ -1176,6 +1176,21 @@ void USkeletalMeshEditorSubsystem::SetMorphTargetsToGeneratedByEngineForAllSkele
 	}
 }
 
+void USkeletalMeshEditorSubsystem::GetMorphTargetsGeneratedByEngine(USkeletalMesh* TargetMesh, TArray<FName>& OutNames)
+{
+	constexpr int32 LodIndex = 0;
+	if (FSkeletalMeshLODInfo* LodInfo = TargetMesh->GetLODInfo(LodIndex))
+	{
+		for (const TPair<FString, FMorphTargetImportedSourceFileInfo>& Info : LodInfo->ImportedMorphTargetSourceFilename)
+		{
+			if (Info.Value.IsGeneratedByEngine())
+			{
+				OutNames.Add(*(Info.Key));
+			}
+		}
+	}
+}
+
 void USkeletalMeshEditorSubsystem::GetSkeletonCurveMetaDataNames(const USkeleton* TargetSkeleton, TArray<FName>& OutNames, ESkelSubSysQueryCurvesMetatdataNamesFilter Filter)
 {
 	if (!TargetSkeleton)
diff --git a/Engine/Source/Editor/SkeletalMeshEditor/Public/ISkeletalMeshEditor.h b/Engine/Source/Editor/SkeletalMeshEditor/Public/ISkeletalMeshEditor.h
--- a/Engine/Source/Editor/SkeletalMeshEditor/Public/ISkeletalMeshEditor.h
+++ b/Engine/Source/Editor/SkeletalMeshEditor/Public/ISkeletalMeshEditor.h
@@ -12,5 +12,10 @@ class ISkeletalMeshEditor : public FPersonaAssetEditorToolkit, public IHasPerson
 {
 public:
 	virtual TSharedPtr<ISkeletalMeshEditorBinding> GetBinding() = 0;
+	virtual FSimpleMulticastDelegate& OnPreSaveAsset() = 0;
+	virtual FSimpleMulticastDelegate& OnPreSaveAssetAs() = 0;
+
+	virtual TSharedPtr<FUICommandInfo> GetResetBoneTransformsCommand() = 0;
+	virtual TSharedPtr<FUICommandInfo> GetResetAllBonesTransformsCommand() = 0;
 };
 
diff --git a/Engine/Source/Editor/SkeletalMeshEditor/Public/ISkeletalMeshEditorModule.h b/Engine/Source/Editor/SkeletalMeshEditor/Public/ISkeletalMeshEditorModule.h
--- a/Engine/Source/Editor/SkeletalMeshEditor/Public/ISkeletalMeshEditorModule.h
+++ b/Engine/Source/Editor/SkeletalMeshEditor/Public/ISkeletalMeshEditorModule.h
@@ -24,4 +24,8 @@ public:
 	/** Add callback after editor has been initialized */
 	DECLARE_DELEGATE_OneParam(FOnSkeletalMeshEditorInitialized, TWeakPtr<ISkeletalMeshEditor>);
 	virtual TArray<FOnSkeletalMeshEditorInitialized>& GetPostEditorInitDelegates() = 0;
+
+	// Used by the EditorUISubsystem
+	DECLARE_EVENT_OneParam(FExampleCharacterFXEditorModule, FOnRegisterLayoutExtensions, FLayoutExtender&);
+	virtual FOnRegisterLayoutExtensions& OnRegisterLayoutExtensions() = 0;
 };
diff --git a/Engine/Source/Editor/SkeletalMeshEditor/Public/SkeletalMeshEditorSubsystem.h b/Engine/Source/Editor/SkeletalMeshEditor/Public/SkeletalMeshEditorSubsystem.h
--- a/Engine/Source/Editor/SkeletalMeshEditor/Public/SkeletalMeshEditorSubsystem.h
+++ b/Engine/Source/Editor/SkeletalMeshEditor/Public/SkeletalMeshEditorSubsystem.h
@@ -343,6 +343,15 @@ public:
 	UFUNCTION(BlueprintCallable, Category = "Editor Scripting | SkeletalMesh")
 	static UE_API void SetMorphTargetsToGeneratedByEngineForAllSkeletalMesh(const TArray<FString>& OptionalNames, const TArray<FName>& OptionalPaths);
 
+	/**
+	* Get morph targets marked as generated by engine.
+	* 
+	* @param TargetMesh The mesh to get morph targets.
+	* @param OutNames - Receive the names of morph targets marked as generated by engine
+	*/
+	UFUNCTION(BlueprintCallable, Category = "Editor Scripting | SkeletalMesh")
+	static UE_API void GetMorphTargetsGeneratedByEngine(USkeletalMesh* TargetMesh, TArray<FName>& OutNames);
+	
 	/**
 	 * Retrieve all curve metadata name from the skeleton.
 	 * 
diff --git a/Engine/Source/Editor/SkeletalMeshEditor/SkeletalMeshEditor.Build.cs b/Engine/Source/Editor/SkeletalMeshEditor/SkeletalMeshEditor.Build.cs
--- a/Engine/Source/Editor/SkeletalMeshEditor/SkeletalMeshEditor.Build.cs
+++ b/Engine/Source/Editor/SkeletalMeshEditor/SkeletalMeshEditor.Build.cs
@@ -40,7 +40,8 @@ public class SkeletalMeshEditor : ModuleRules
 				"StatusBar",
 				"PhysicsUtilities",
 				"InterchangeCore",
-				"InterchangeEngine"
+				"InterchangeEngine",
+				"WorkspaceMenuStructure"
 			}
 		);
 	}
diff --git a/Engine/Source/Runtime/MeshConversion/Private/DynamicMeshToMeshDescription.cpp b/Engine/Source/Runtime/MeshConversion/Private/DynamicMeshToMeshDescription.cpp
--- a/Engine/Source/Runtime/MeshConversion/Private/DynamicMeshToMeshDescription.cpp
+++ b/Engine/Source/Runtime/MeshConversion/Private/DynamicMeshToMeshDescription.cpp
@@ -839,10 +839,16 @@ void FDynamicMeshToMeshDescription::Convert_NoSharedInstances(const FDynamicMesh
 		}
 	}
 
-	TMap<FName, FSkinWeightsVertexAttributesRef> VertexBoneWeightsMap; 
+	TMap<FName, FSkinWeightsVertexAttributesRef> VertexBoneWeightsMap;
 	if (bHasSkinWeightsAttributes)
 	{
 		FSkeletalMeshAttributes MeshOutAttributes(MeshOut);
+		
+		for (const FName Name : MeshOutAttributes.GetSkinWeightProfileNames())
+		{
+			MeshOutAttributes.UnregisterSkinWeightAttribute(Name);
+		}
+		
 		for (const TTuple<FName, TUniquePtr<FDynamicMeshVertexSkinWeightsAttribute>>& AttributeInfo: MeshIn->Attributes()->GetSkinWeightsAttributes())
 		{
 			FName ProfileName = AttributeInfo.Key;
@@ -856,6 +862,12 @@ void FDynamicMeshToMeshDescription::Convert_NoSharedInstances(const FDynamicMesh
 	if (bHasMorphTargetAttributes)
 	{
 		FSkeletalMeshAttributes MeshOutAttributes(MeshOut);
+		
+		for (const FName Name : MeshOutAttributes.GetMorphTargetNames())
+		{
+			MeshOutAttributes.UnregisterMorphTargetAttribute(Name);
+		}
+		
 		for (const TTuple<FName, TUniquePtr<FDynamicMeshMorphTargetAttribute>>& AttributeInfo: MeshIn->Attributes()->GetMorphTargetAttributes())
 		{
 			FName MorphTargetName = AttributeInfo.Key;
@@ -863,7 +875,6 @@ void FDynamicMeshToMeshDescription::Convert_NoSharedInstances(const FDynamicMesh
 			MeshOutAttributes.RegisterMorphTargetAttribute(MorphTargetName, false);
 			MorphTargetsMap.Add(MorphTargetName, MeshOutAttributes.GetVertexMorphPositionDelta(MorphTargetName));
 		}
-		
 	}
 
 	// always copy when we are baking new mesh? should this be a config option?
