[SKMModelingtools] adding a missing file from a previous change

diff --git a/Engine/Plugins/Animation/SkeletalMeshModelingTools/Source/SkeletalMeshModelingTools/Private/SkeletalMeshModelingToolsCommands.cpp b/Engine/Plugins/Animation/SkeletalMeshModelingTools/Source/SkeletalMeshModelingTools/Private/SkeletalMeshModelingToolsCommands.cpp
--- a/Engine/Plugins/Animation/SkeletalMeshModelingTools/Source/SkeletalMeshModelingTools/Private/SkeletalMeshModelingToolsCommands.cpp
+++ b/Engine/Plugins/Animation/SkeletalMeshModelingTools/Source/SkeletalMeshModelingTools/Private/SkeletalMeshModelingToolsCommands.cpp
@@ -23,6 +23,8 @@ void FSkeletalMeshModelingToolsCommands::RegisterCommands()
 	UI_COMMAND(CopyBones, "Copy Bone(s)", "Copy selected bone(s).", EUserInterfaceActionType::Button, FInputChord(EKeys::C, EModifierKey::Control));
 	UI_COMMAND(PasteBones, "Paste Bone(s)", "Paste selected bone(s).", EUserInterfaceActionType::Button, FInputChord(EKeys::V, EModifierKey::Control));
 	UI_COMMAND(DuplicateBones, "Duplicate Bone(s)", "Duplicate selected bone(s).", EUserInterfaceActionType::Button, FInputChord(EKeys::D, EModifierKey::Control));
+	
+	UI_COMMAND(NewMorphTarget, "New Morph Target", "Add new Morph Target.", EUserInterfaceActionType::Button, FInputChord(EKeys::N, EModifierKey::Control));
 }
 
 const FSkeletalMeshModelingToolsCommands& FSkeletalMeshModelingToolsCommands::Get()
