[SKMModellingTool] Fixed overlapping hotkey for commands

diff --git a/Engine/Plugins/Animation/SkeletalMeshModelingTools/Source/SkeletalMeshModelingTools/Private/SkeletalMeshModelingToolsCommands.cpp b/Engine/Plugins/Animation/SkeletalMeshModelingTools/Source/SkeletalMeshModelingTools/Private/SkeletalMeshModelingToolsCommands.cpp
--- a/Engine/Plugins/Animation/SkeletalMeshModelingTools/Source/SkeletalMeshModelingTools/Private/SkeletalMeshModelingToolsCommands.cpp
+++ b/Engine/Plugins/Animation/SkeletalMeshModelingTools/Source/SkeletalMeshModelingTools/Private/SkeletalMeshModelingToolsCommands.cpp
@@ -24,7 +24,7 @@ void FSkeletalMeshModelingToolsCommands::RegisterCommands()
 	UI_COMMAND(PasteBones, "Paste Bone(s)", "Paste selected bone(s).", EUserInterfaceActionType::Button, FInputChord(EKeys::V, EModifierKey::Control));
 	UI_COMMAND(DuplicateBones, "Duplicate Bone(s)", "Duplicate selected bone(s).", EUserInterfaceActionType::Button, FInputChord(EKeys::D, EModifierKey::Control));
 	
-	UI_COMMAND(NewMorphTarget, "New Morph Target", "Add new Morph Target.", EUserInterfaceActionType::Button, FInputChord(EKeys::N, EModifierKey::Control));
+	UI_COMMAND(NewMorphTarget, "New Morph Target", "Add new Morph Target.", EUserInterfaceActionType::Button, FInputChord());
 }
 
 const FSkeletalMeshModelingToolsCommands& FSkeletalMeshModelingToolsCommands::Get()
