[SKMModellingTools] added additional tool tips for "Apply Manually" Mode, also changed a few UI text to make things clearer

diff --git a/Engine/Plugins/Animation/SkeletalMeshModelingTools/Source/SkeletalMeshModelingTools/Private/SkeletalMeshModelingToolsEditorMode.cpp b/Engine/Plugins/Animation/SkeletalMeshModelingTools/Source/SkeletalMeshModelingTools/Private/SkeletalMeshModelingToolsEditorMode.cpp
--- a/Engine/Plugins/Animation/SkeletalMeshModelingTools/Source/SkeletalMeshModelingTools/Private/SkeletalMeshModelingToolsEditorMode.cpp
+++ b/Engine/Plugins/Animation/SkeletalMeshModelingTools/Source/SkeletalMeshModelingTools/Private/SkeletalMeshModelingToolsEditorMode.cpp
@@ -1197,6 +1197,14 @@ USkeletalMeshModelingToolsEditorMode::EApplyMode USkeletalMeshModelingToolsEdito
 void USkeletalMeshModelingToolsEditorMode::SetApplyMode(EApplyMode InApplyMode)
 {
 	ApplyMode = InApplyMode;
+	if (ApplyMode == EApplyMode::ApplyManually)
+	{
+		GetInteractiveToolsContext()->ToolManager->SetToolSwitchMode(EToolManagerToolSwitchMode::AcceptIfAble);
+	}
+	else
+	{
+		GetInteractiveToolsContext()->ToolManager->SetToolSwitchMode(EToolManagerToolSwitchMode::CancelIfAble);
+	}
 }
 
 void USkeletalMeshModelingToolsEditorMode::HideSkeletonForTool()
diff --git a/Engine/Plugins/Animation/SkeletalMeshModelingTools/Source/SkeletalMeshModelingTools/Private/SkeletalMeshModelingToolsEditorModeToolkit.cpp b/Engine/Plugins/Animation/SkeletalMeshModelingTools/Source/SkeletalMeshModelingTools/Private/SkeletalMeshModelingToolsEditorModeToolkit.cpp
--- a/Engine/Plugins/Animation/SkeletalMeshModelingTools/Source/SkeletalMeshModelingTools/Private/SkeletalMeshModelingToolsEditorModeToolkit.cpp
+++ b/Engine/Plugins/Animation/SkeletalMeshModelingTools/Source/SkeletalMeshModelingTools/Private/SkeletalMeshModelingToolsEditorModeToolkit.cpp
@@ -423,8 +423,31 @@ void FSkeletalMeshModelingToolsEditorModeToolkit::MakeToolAcceptCancelWidget()
 			.Padding(FMargin(0.0, 0.f, 2.f, 0.f))
 			[
 				SNew(SPrimaryButton)
-				.Text(LOCTEXT("OverlayAccept", "Accept"))
-				.ToolTipText(LOCTEXT("OverlayAcceptTooltip", "Accept/Commit the results of the active Tool [Enter]"))
+				.Text_Lambda([this]()
+					{
+						if (EditorMode->GetApplyMode() == USkeletalMeshModelingToolsEditorMode::EApplyMode::ApplyOnToolExit)
+						{
+							return LOCTEXT("OverlayApplyToAsset", "Apply To Asset");
+						}
+						
+						check(EditorMode->GetApplyMode() == USkeletalMeshModelingToolsEditorMode::EApplyMode::ApplyManually)
+						
+						return LOCTEXT("OverlayExitTool", "Exit Tool");
+					})
+				.ToolTipText_Lambda([this]()
+				{
+					if (EditorMode->GetApplyMode() == USkeletalMeshModelingToolsEditorMode::EApplyMode::ApplyOnToolExit)
+					{
+						return LOCTEXT("OverlayApplyToAssetToolTip",
+							"\"Apply on Tool Exit\" mode is active: changes are applied to the asset immediately. "
+							"Switching to a different tool discards current change");
+					}
+						
+					check(EditorMode->GetApplyMode() == USkeletalMeshModelingToolsEditorMode::EApplyMode::ApplyManually)
+						
+					return LOCTEXT("OverlayExitToolToolTip",
+						"\"Apply Manually\" mode is active: changes are saved when Exiting tool or Switching to a different tool");
+				})
 				.OnClicked_Lambda([this]() { GetScriptableEditorMode()->GetInteractiveToolsContext()->EndTool(EToolShutdownType::Accept); return FReply::Handled(); })
 				.IsEnabled_Lambda([this]() { return GetScriptableEditorMode()->GetInteractiveToolsContext()->CanAcceptActiveTool(); })
 				.Visibility_Lambda([this]() { return GetScriptableEditorMode()->GetInteractiveToolsContext()->ActiveToolHasAccept() ? EVisibility::Visible : EVisibility::Collapsed; })
@@ -689,10 +712,12 @@ TSharedPtr<SWidget> FSkeletalMeshModelingToolsEditorModeToolkit::MakeFooterWidge
 
 	// Apply Mode Picker Widget
 	{
+		FText ToolTip = LOCTEXT("ApplyModeToolTip", "Select when tool changes should be applied to the asset.\n"
+								"When using \"Apply Manually\", you can switch between different tools directly without waiting for changes to be applied to the asset");
 		ApplyModeLabel =
 			SNew(STextBlock)
 			.Text(LOCTEXT("ApplyModeLabel", "Apply Mode"))
-			.ToolTipText(LOCTEXT("ApplyModeToolTip", "Select when tool changes should be applied to the asset"));
+			.ToolTipText(ToolTip);
 
 		ApplyModeEnums.Add(USkeletalMeshModelingToolsEditorMode::EApplyMode::ApplyOnToolExit);
 		ApplyModeStrings.Add(MakeShared<FString>(TEXT("Apply On Tool Exit")));
@@ -702,6 +727,7 @@ TSharedPtr<SWidget> FSkeletalMeshModelingToolsEditorModeToolkit::MakeFooterWidge
 			
 		ApplyModeComboBox =
 			SNew(STextComboBox)
+			.ToolTipText(ToolTip)
 			.OptionsSource(&ApplyModeStrings)
 			.IsEnabled(this, &FSkeletalMeshModelingToolsEditorModeToolkit::CanChangeAssetEditingSettings)
 			.OnSelectionChanged_Lambda([this](TSharedPtr<FString> String, ESelectInfo::Type)
@@ -839,7 +865,7 @@ FText FSkeletalMeshModelingToolsEditorModeToolkit::GetApplyButtonText() const
 		return FText::Format(LOCTEXT("ApplyButtonZeroChangeText", "No Pending Changes"), FText::AsNumber(ChangeCount));
 	}
 	
-	return FText::Format(LOCTEXT("ApplyButtonChangeText", "Apply {0}|plural(one=Change, other=Changes)"), FText::AsNumber(ChangeCount));
+	return FText::Format(LOCTEXT("ApplyButtonChangeText", "Apply {0} {0}|plural(one=Change, other=Changes)"), FText::AsNumber(ChangeCount));
 }
 
 bool FSkeletalMeshModelingToolsEditorModeToolkit::IsDiscardButtonEnabled() const
diff --git a/Engine/Source/Editor/Persona/Private/TabSpawners.cpp b/Engine/Source/Editor/Persona/Private/TabSpawners.cpp
--- a/Engine/Source/Editor/Persona/Private/TabSpawners.cpp
+++ b/Engine/Source/Editor/Persona/Private/TabSpawners.cpp
@@ -157,13 +157,13 @@ FMorphTargetTabSummoner::FMorphTargetTabSummoner(TSharedPtr<class FAssetEditorTo
 	: FWorkflowTabFactory(FPersonaTabs::MorphTargetsID, InHostingApp)
 	, ViewerWidget(InViewerWidget)
 {
-	TabLabel = LOCTEXT("MorphTargetTabTitle", "Morph Target Previewer");
+	TabLabel = LOCTEXT("MorphTargetTabTitle", "Morph Targets");
 	TabIcon = FSlateIcon(FAppStyle::GetAppStyleSetName(), "Persona.Tabs.MorphTargetPreviewer");
 
 	EnableTabPadding();
 	bIsSingleton = true;
 
-	ViewMenuDescription = LOCTEXT("MorphTargetTabView", "Morph Target Previewer");
+	ViewMenuDescription = LOCTEXT("MorphTargetTabView", "Morph Targets");
 	ViewMenuTooltip = LOCTEXT("MorphTargetTabView_ToolTip", "Shows the morph target viewer");
 }
 
