[SKMModelingTool] make sure ref skeleton tree's default is writable instead of readonly

diff --git a/Engine/Plugins/Animation/SkeletalMeshModelingTools/Source/SkeletalMeshModelingTools/Private/SReferenceSkeletonTree.h b/Engine/Plugins/Animation/SkeletalMeshModelingTools/Source/SkeletalMeshModelingTools/Private/SReferenceSkeletonTree.h
--- a/Engine/Plugins/Animation/SkeletalMeshModelingTools/Source/SkeletalMeshModelingTools/Private/SReferenceSkeletonTree.h
+++ b/Engine/Plugins/Animation/SkeletalMeshModelingTools/Source/SkeletalMeshModelingTools/Private/SReferenceSkeletonTree.h
@@ -181,7 +181,9 @@ private:
 class SReferenceSkeletonTree : public SCompoundWidget, public FEditorUndoClient
 {
 public:
-	SLATE_BEGIN_ARGS(SReferenceSkeletonTree) {}
+	SLATE_BEGIN_ARGS(SReferenceSkeletonTree)
+		: _bIsReadOnly(false)
+		{}
 		SLATE_ARGUMENT(TWeakObjectPtr<USkeletonModifier>, Modifier)
 		SLATE_ARGUMENT(FRefSkeletonTreeDelegates, Delegates)
 		SLATE_ARGUMENT(bool, bIsReadOnly)
