[SKMModelingTool] Fixed a crash when exiting skeletong edit tool if a bone is deleted

diff --git a/Engine/Plugins/Animation/SkeletalMeshModelingTools/Source/SkeletalMeshModelingTools/Private/Components/SKMBackedDynaMeshComponent.cpp b/Engine/Plugins/Animation/SkeletalMeshModelingTools/Source/SkeletalMeshModelingTools/Private/Components/SKMBackedDynaMeshComponent.cpp
--- a/Engine/Plugins/Animation/SkeletalMeshModelingTools/Source/SkeletalMeshModelingTools/Private/Components/SKMBackedDynaMeshComponent.cpp
+++ b/Engine/Plugins/Animation/SkeletalMeshModelingTools/Source/SkeletalMeshModelingTools/Private/Components/SKMBackedDynaMeshComponent.cpp
@@ -34,8 +34,11 @@ void USkeletalMeshBackedDynamicMeshComponent::FSkeletonChangeTracker::HandleSkel
 	for (int32 BoneIndex = 0; BoneIndex < BoneIndexTracker.Num(); BoneIndex++)
 	{
 		int32 PreviousNewBoneIndex = BoneIndexTracker[BoneIndex];
-		int32 NewBoneIndex = ToolBoneIndexTracker[PreviousNewBoneIndex];
-		BoneIndexTracker[BoneIndex] = NewBoneIndex;
+		if (ToolBoneIndexTracker.IsValidIndex(PreviousNewBoneIndex))
+		{
+			int32 NewBoneIndex = ToolBoneIndexTracker[PreviousNewBoneIndex];
+			BoneIndexTracker[BoneIndex] = NewBoneIndex;
+		}
 	}
 
 	ChangeCount++;
