[SKMModeling tool] fixed discard changes not redoable

diff --git a/Engine/Plugins/Animation/SkeletalMeshModelingTools/Source/SkeletalMeshModelingTools/Private/Components/SKMBackedDynaMeshComponent.cpp b/Engine/Plugins/Animation/SkeletalMeshModelingTools/Source/SkeletalMeshModelingTools/Private/Components/SKMBackedDynaMeshComponent.cpp
--- a/Engine/Plugins/Animation/SkeletalMeshModelingTools/Source/SkeletalMeshModelingTools/Private/Components/SKMBackedDynaMeshComponent.cpp
+++ b/Engine/Plugins/Animation/SkeletalMeshModelingTools/Source/SkeletalMeshModelingTools/Private/Components/SKMBackedDynaMeshComponent.cpp
@@ -331,7 +331,7 @@ bool USkeletalMeshBackedDynamicMeshComponent::IsCommitting() const
 void USkeletalMeshBackedDynamicMeshComponent::DiscardChanges()
 {
 	// Records a change that turns current mesh back to the initial state
-	FChangeScope(this);
+	FChangeScope ChangeScope(this);
 	
 	FDynamicMesh3 InitialMeshCopy = *Trackers.InitialAssetMesh;
 	SetMesh(MoveTemp(InitialMeshCopy));
