[SKMModelingTool] Fixed crash adding new morph target after adding a bone in immediate apply mode

diff --git a/Engine/Plugins/Animation/SkeletalMeshModelingTools/Source/SkeletalMeshModelingTools/Private/Components/SKMBackedDynaMeshComponent.cpp b/Engine/Plugins/Animation/SkeletalMeshModelingTools/Source/SkeletalMeshModelingTools/Private/Components/SKMBackedDynaMeshComponent.cpp
--- a/Engine/Plugins/Animation/SkeletalMeshModelingTools/Source/SkeletalMeshModelingTools/Private/Components/SKMBackedDynaMeshComponent.cpp
+++ b/Engine/Plugins/Animation/SkeletalMeshModelingTools/Source/SkeletalMeshModelingTools/Private/Components/SKMBackedDynaMeshComponent.cpp
@@ -345,6 +345,7 @@ void USkeletalMeshBackedDynamicMeshComponent::HandleSkeletonChange(USkeletonModi
 	FTrackerChangeScope ChangeScope(this);
 	
 	Trackers.RefSkeleton = InModifier->GetReferenceSkeleton();
+	Trackers.RefSkeleton.GetBoneAbsoluteTransforms(Trackers.ComponentSpaceBoneTransformsRefPose);
 	Trackers.SkeletonChangeTracker.HandleSkeletonChanged(InModifier->GetBoneIndexTracker());
 	MarkDirtyInternal();
 }
@@ -535,7 +536,7 @@ void USkeletalMeshBackedDynamicMeshComponent::ResetTrackersDirect()
 	Trackers.ChangeCount = 0;
 	
 	Trackers.RefSkeleton = WeakSkeletalMesh->GetRefSkeleton();
-
+	Trackers.RefSkeleton.GetBoneAbsoluteTransforms(Trackers.ComponentSpaceBoneTransformsRefPose);
 	Trackers.SkeletonChangeTracker.Init(Trackers.RefSkeleton.GetRawBoneNum());
 	
 	TArray<FName> MorphTargetNames;
diff --git a/Engine/Plugins/Animation/SkeletalMeshModelingTools/Source/SkeletalMeshModelingTools/Private/Components/SKMBackedDynaMeshComponent.h b/Engine/Plugins/Animation/SkeletalMeshModelingTools/Source/SkeletalMeshModelingTools/Private/Components/SKMBackedDynaMeshComponent.h
--- a/Engine/Plugins/Animation/SkeletalMeshModelingTools/Source/SkeletalMeshModelingTools/Private/Components/SKMBackedDynaMeshComponent.h
+++ b/Engine/Plugins/Animation/SkeletalMeshModelingTools/Source/SkeletalMeshModelingTools/Private/Components/SKMBackedDynaMeshComponent.h
@@ -85,6 +85,7 @@ public:
 	const FSkeletonChangeTracker& GetSkeletonChangeTracker() const {return Trackers.SkeletonChangeTracker;}
 	const FMorphTargetChangeTracker& GetMorphTargetChangeTracker() const {return Trackers.MorphTargetChangeTracker;}
 	const FReferenceSkeleton& GetRefSkeleton() const {return Trackers.RefSkeleton;}
+	const TArray<FTransform>& GetComponentSpaceBoneTransformsRefPose() const {return Trackers.ComponentSpaceBoneTransformsRefPose;}
 
 	bool IsSkeletonDirty() const;
 	bool IsDirty() const;
@@ -111,6 +112,7 @@ protected:
 		int32 ChangeCount = 0;
 		
 		FReferenceSkeleton RefSkeleton;
+		TArray<FTransform> ComponentSpaceBoneTransformsRefPose;
 		FSkeletonChangeTracker SkeletonChangeTracker;
 		FMorphTargetChangeTracker MorphTargetChangeTracker;
 	};
diff --git a/Engine/Plugins/Animation/SkeletalMeshModelingTools/Source/SkeletalMeshModelingTools/Private/SkeletalMeshEditingCache.cpp b/Engine/Plugins/Animation/SkeletalMeshModelingTools/Source/SkeletalMeshModelingTools/Private/SkeletalMeshEditingCache.cpp
--- a/Engine/Plugins/Animation/SkeletalMeshModelingTools/Source/SkeletalMeshModelingTools/Private/SkeletalMeshEditingCache.cpp
+++ b/Engine/Plugins/Animation/SkeletalMeshModelingTools/Source/SkeletalMeshModelingTools/Private/SkeletalMeshEditingCache.cpp
@@ -708,7 +708,7 @@ URefSkeletonPoser* USkeletalMeshEditingCache::GetSkeletonPoser() const
 
 const TArray<FTransform>& USkeletalMeshEditingCache::GetComponentSpaceBoneTransformsRefPose() const
 {
-	return RefSkeletonPoser->GetComponentSpaceTransformsRefPose();
+	return EditingMeshComponent->GetComponentSpaceBoneTransformsRefPose();
 }
 
 const TMap<FName, float>& USkeletalMeshEditingCache::GetSkeletalMeshComponentMorphTargetWeights() const
@@ -728,6 +728,11 @@ void USkeletalMeshEditingCache::DeformPreviewMesh(const TArray<FTransform>& Comp
 	{
 		return;
 	}
+
+	if (GetComponentSpaceBoneTransformsRefPose().Num() != ComponentSpaceTransforms.Num())
+	{
+		return;
+	}
 	
 	using namespace UE::Geometry;
 	using namespace SkeletalMeshToolsHelper;
diff --git a/Engine/Plugins/Runtime/MeshModelingToolset/Source/MeshModelingToolsEditorOnly/Private/SkeletalMesh/RefSkeletonPoser.cpp b/Engine/Plugins/Runtime/MeshModelingToolset/Source/MeshModelingToolsEditorOnly/Private/SkeletalMesh/RefSkeletonPoser.cpp
--- a/Engine/Plugins/Runtime/MeshModelingToolset/Source/MeshModelingToolsEditorOnly/Private/SkeletalMesh/RefSkeletonPoser.cpp
+++ b/Engine/Plugins/Runtime/MeshModelingToolset/Source/MeshModelingToolsEditorOnly/Private/SkeletalMesh/RefSkeletonPoser.cpp
@@ -50,7 +50,7 @@ void URefSkeletonPoser::SetRefSkeleton(const FReferenceSkeleton& InRefSkeleton)
 	RefSkeleton = InRefSkeleton;
 	
 	Transforms = RefSkeleton.GetRefBonePose();
-	TransformCached = TBitArray<>(false, RefSkeleton.GetRawBoneNum());
+	TransformCached = TBitArray<>(false, RefSkeleton.GetNum());
 	
 	RefSkeleton.GetBoneAbsoluteTransforms(RefTransforms);
 
