[SKMModellingTool] Fixed some text

#DisabledValidations nonUnityAndStaticAnalysis

diff --git a/Engine/Plugins/Animation/SkeletalMeshModelingTools/Source/SkeletalMeshModelingTools/Private/SkeletalMeshModelingToolsEditorModeToolkit.cpp b/Engine/Plugins/Animation/SkeletalMeshModelingTools/Source/SkeletalMeshModelingTools/Private/SkeletalMeshModelingToolsEditorModeToolkit.cpp
--- a/Engine/Plugins/Animation/SkeletalMeshModelingTools/Source/SkeletalMeshModelingTools/Private/SkeletalMeshModelingToolsEditorModeToolkit.cpp
+++ b/Engine/Plugins/Animation/SkeletalMeshModelingTools/Source/SkeletalMeshModelingTools/Private/SkeletalMeshModelingToolsEditorModeToolkit.cpp
@@ -427,7 +427,7 @@ void FSkeletalMeshModelingToolsEditorModeToolkit::MakeToolAcceptCancelWidget()
 					{
 						if (EditorMode->GetApplyMode() == USkeletalMeshModelingToolsEditorMode::EApplyMode::ApplyOnToolExit)
 						{
-							return LOCTEXT("OverlayApplyToAsset", "Apply To Asset");
+							return LOCTEXT("OverlayApplyToAsset", "Apply to Asset");
 						}
 						
 						check(EditorMode->GetApplyMode() == USkeletalMeshModelingToolsEditorMode::EApplyMode::ApplyManually)
