[SKM Modeling Tools] Fixed widget mode incorrectly greyed out when using skeleton tool with dyna mesh skeleton

No need to override widget mode checking as all widget modes exposed to the UI are supported.

diff --git a/Engine/Plugins/Animation/SkeletalMeshModelingTools/Source/SkeletalMeshModelingTools/Private/SkeletalMeshModelingToolsEditorMode.cpp b/Engine/Plugins/Animation/SkeletalMeshModelingTools/Source/SkeletalMeshModelingTools/Private/SkeletalMeshModelingToolsEditorMode.cpp
--- a/Engine/Plugins/Animation/SkeletalMeshModelingTools/Source/SkeletalMeshModelingTools/Private/SkeletalMeshModelingToolsEditorMode.cpp
+++ b/Engine/Plugins/Animation/SkeletalMeshModelingTools/Source/SkeletalMeshModelingTools/Private/SkeletalMeshModelingToolsEditorMode.cpp
@@ -947,18 +947,6 @@ bool USkeletalMeshModelingToolsEditorMode::UsesTransformWidget() const
 	return false;
 }
 
-bool USkeletalMeshModelingToolsEditorMode::UsesTransformWidget(UE::Widget::EWidgetMode CheckMode) const
-{
-	// This function determines if the specific widget mode can be used
-	
-	if (GetCurrentEditingCache()->IsDynamicMeshSkeletonEnabled())
-	{
-		return ShouldDrawWidget() && (CheckMode == UE::Widget::WM_Scale || CheckMode == UE::Widget::WM_Translate || CheckMode == UE::Widget::WM_Rotate);
-	}
-	
-	return Super::UsesTransformWidget(CheckMode);
-}
-
 FVector USkeletalMeshModelingToolsEditorMode::GetWidgetLocation() const
 {
 	if (GetCurrentEditingCache()->IsDynamicMeshSkeletonEnabled())
diff --git a/Engine/Plugins/Animation/SkeletalMeshModelingTools/Source/SkeletalMeshModelingTools/Private/SkeletalMeshModelingToolsEditorMode.h b/Engine/Plugins/Animation/SkeletalMeshModelingTools/Source/SkeletalMeshModelingTools/Private/SkeletalMeshModelingToolsEditorMode.h
--- a/Engine/Plugins/Animation/SkeletalMeshModelingTools/Source/SkeletalMeshModelingTools/Private/SkeletalMeshModelingToolsEditorMode.h
+++ b/Engine/Plugins/Animation/SkeletalMeshModelingTools/Source/SkeletalMeshModelingTools/Private/SkeletalMeshModelingToolsEditorMode.h
@@ -101,7 +101,6 @@ public:
 	virtual bool AllowWidgetMove() override;
 	virtual bool ShouldDrawWidget() const override;
 	virtual bool UsesTransformWidget() const override;
-	virtual bool UsesTransformWidget(UE::Widget::EWidgetMode CheckMode) const override;
 	virtual FVector GetWidgetLocation() const override;
 	virtual bool GetCustomDrawingCoordinateSystem(FMatrix& InMatrix, void* InData) override;
 	virtual bool GetCustomInputCoordinateSystem(FMatrix& InMatrix, void* InData) override;
