[SKM modeling tools] Added more null checks

diff --git a/Engine/Plugins/Animation/SkeletalMeshModelingTools/Source/SkeletalMeshModelingTools/Private/ToolTargets/SKMBackedDynaMeshComponentToolTarget.cpp b/Engine/Plugins/Animation/SkeletalMeshModelingTools/Source/SkeletalMeshModelingTools/Private/ToolTargets/SKMBackedDynaMeshComponentToolTarget.cpp
--- a/Engine/Plugins/Animation/SkeletalMeshModelingTools/Source/SkeletalMeshModelingTools/Private/ToolTargets/SKMBackedDynaMeshComponentToolTarget.cpp
+++ b/Engine/Plugins/Animation/SkeletalMeshModelingTools/Source/SkeletalMeshModelingTools/Private/ToolTargets/SKMBackedDynaMeshComponentToolTarget.cpp
@@ -13,12 +13,21 @@
 
 FReferenceSkeleton USkeletalMeshBackedDynamicMeshComponentToolTarget::GetSkeleton() const
 {
-	return GetComponent()->GetRefSkeleton();
+	if (!GetComponent())
+	{
+		return {};
+	}
 	
+	return GetComponent()->GetRefSkeleton();
 }
 
 void USkeletalMeshBackedDynamicMeshComponentToolTarget::CommitSkeletonModifier(USkeletonModifier* InModifier)
 {
+	if (!GetComponent())
+	{
+		return;
+	}
+	
 	Super::CommitSkeletonModifier(InModifier);
 
 	GetComponent()->HandleSkeletonChange(InModifier);
@@ -31,6 +40,11 @@ void USkeletalMeshBackedDynamicMeshComponentToolTarget::CommitSkeletonModifier(U
 
 void USkeletalMeshBackedDynamicMeshComponentToolTarget::CommitDynamicMeshChange(TUniquePtr<FToolCommandChange> Change, const FText& ChangeMessage)
 {
+	if (!GetComponent())
+	{
+		return;
+	}
+	
 	Super::CommitDynamicMeshChange(MoveTemp(Change), ChangeMessage);
 	GetComponent()->MarkDirty();
 
@@ -42,16 +56,31 @@ void USkeletalMeshBackedDynamicMeshComponentToolTarget::CommitDynamicMeshChange(
 
 USkeletalMesh* USkeletalMeshBackedDynamicMeshComponentToolTarget::GetSkeletalMesh() const
 {
+	if (!GetComponent())
+	{
+		return nullptr;
+	}
+	
 	return GetComponent()->GetSkeletalMesh();
 }
 
 void USkeletalMeshBackedDynamicMeshComponentToolTarget::SetOwnerVisibility(bool bVisible) const
 {
+	if (!GetComponent())
+	{
+		return;
+	}
+	
 	GetComponent()->ForwardVisibilityChangeRequest(bVisible);
 }
 
 USkeletalMeshBackedDynamicMeshComponent* USkeletalMeshBackedDynamicMeshComponentToolTarget::GetComponent() const
 {
+	if (!Component.IsValid())
+	{
+		return nullptr;
+	}
+
 	return CastChecked<USkeletalMeshBackedDynamicMeshComponent>(Component);
 }
 
