[SKMModellingTool] Fixed skin weight paint tool not showing its custom skeleton tree widget when user change to it while having the skeleton edit tool active (with skeleton edited). Also the edited skeleton wasn't refected in the custom skeleton tree.

diff --git a/Engine/Plugins/Animation/SkeletalMeshModelingTools/Source/SkeletalMeshModelingTools/Private/SkeletalMeshModelingToolsEditorMode.cpp b/Engine/Plugins/Animation/SkeletalMeshModelingTools/Source/SkeletalMeshModelingTools/Private/SkeletalMeshModelingToolsEditorMode.cpp
--- a/Engine/Plugins/Animation/SkeletalMeshModelingTools/Source/SkeletalMeshModelingTools/Private/SkeletalMeshModelingToolsEditorMode.cpp
+++ b/Engine/Plugins/Animation/SkeletalMeshModelingTools/Source/SkeletalMeshModelingTools/Private/SkeletalMeshModelingToolsEditorMode.cpp
@@ -375,30 +375,53 @@ void USkeletalMeshModelingToolsEditorMode::Enter()
 			},
 		0);
 	
-	CacheDynamicMeshSkeletonStatusWatcher.Initialize(
+	SkeletonTreeHostStateWatcher.Initialize(
 		[this]()
 			{
-				if (GetCurrentEditingCache())
-				{
-					return GetCurrentEditingCache()->IsDynamicMeshSkeletonEnabled();
-				}
-				return false;
+				FSkeletonTreeHostState State;
+				State.bEnableDynamicMeshSkeleton = GetCurrentEditingCache() ? GetCurrentEditingCache()->IsDynamicMeshSkeletonEnabled() : false;
+				State.bIsUsingToolSkeletonWidget = bIsUsingToolSkeletonWidget;
+				State.bActiveToolModifiesSkeleton = bActiveToolModifiesSkeleton;
+				return State;
 			},
-		[this](bool bEnableDynamicMeshSkeleton)
+		[this](const FSkeletonTreeHostState& State)
 			{
 				if (TypedToolkit.IsValid())
 				{
-					if (bEnableDynamicMeshSkeleton)
+					if (State.bActiveToolModifiesSkeleton)
 					{
-						TypedToolkit.Pin()->ShowDynamicMeshSkeletonTree();
+						TypedToolkit.Pin()->DisableEditingCacheSkeletonCommands();
 					}
 					else
 					{
-						TypedToolkit.Pin()->ShowSkeletalMeshSkeletonTree();
+						if (State.bEnableDynamicMeshSkeleton)
+						{
+							TypedToolkit.Pin()->EnableDynamicMeshSkeletonCommands();
+						}
+						else
+						{
+							TypedToolkit.Pin()->EnableSkeletalMeshSkeletonCommands();
+						}	
+					}
+					
+					if (State.bIsUsingToolSkeletonWidget)
+					{
+						TypedToolkit.Pin()->ShowToolSkeletonTree();
+					}
+					else
+					{
+						if (State.bEnableDynamicMeshSkeleton)
+						{
+							TypedToolkit.Pin()->ShowDynamicMeshSkeletonTree();
+						}
+						else
+						{
+							TypedToolkit.Pin()->ShowSkeletalMeshSkeletonTree();
+						}
 					}
 				}
 			},
-		false);
+		FSkeletonTreeHostState());
 	
 	CacheSkeletonChangeCountWatcher.Initialize(
 		[this]()
@@ -445,6 +468,8 @@ TSharedPtr<FTabManager> USkeletalMeshModelingToolsEditorMode::GetAssociatedTabMa
 
 void USkeletalMeshModelingToolsEditorMode::BindToolSkeletonTree(ISkeletalMeshEditingInterface* InEditingInterface)
 {
+	bIsUsingToolSkeletonWidget = true;
+	bActiveToolModifiesSkeleton = InEditingInterface->CanModifySkeleton();
 	if (InEditingInterface->CanModifySkeleton())
 	{
 		// If the tool owns/modifies a skeleton, it should directly interface with the tool skeleton tree, so unbind notifier for the editing cache skeleton
@@ -468,6 +493,9 @@ void USkeletalMeshModelingToolsEditorMode::UnbindToolSkeletonTree()
 		GetCurrentEditingCache()->GetNotifier()->HandleNotification({}, ESkeletalMeshNotifyType::HierarchyChanged);
 		GetCurrentEditingCache()->GetNotifier()->HandleNotification(GetSelectedBones(), ESkeletalMeshNotifyType::BonesSelected);
 	}
+
+	bIsUsingToolSkeletonWidget = false;
+	bActiveToolModifiesSkeleton = false;
 }
 
 const TArray<FTransform>& USkeletalMeshModelingToolsEditorMode::GetComponentSpaceBoneTransforms() const
@@ -703,7 +731,7 @@ void USkeletalMeshModelingToolsEditorMode::Tick(FEditorViewportClient* ViewportC
 	}
 
 	CacheChangeCountWatcher.CheckAndUpdate();
-	CacheDynamicMeshSkeletonStatusWatcher.CheckAndUpdate();
+	SkeletonTreeHostStateWatcher.CheckAndUpdate();
 	CacheSkeletonChangeCountWatcher.CheckAndUpdate();
 }
 
@@ -1100,7 +1128,7 @@ void USkeletalMeshModelingToolsEditorMode::SetEditorBinding(const TWeakPtr<ISkel
 	SkeletalMesh = InSkeletalMeshEditor.Pin()->GetPersonaToolkit()->GetMesh();
 	SkeletalMesh->GetOnMeshChanged().AddUObject(this, &USkeletalMeshModelingToolsEditorMode::HandleSkeletalMeshChanged);
 	SkeletonReader = NewObject<USkeletonModifier>();
-	SkeletonReader->SetSkeletalMesh(SkeletalMesh.Get());
+	SkeletonReader->SetReferenceSkeleton(SkeletalMesh->GetRefSkeleton());
 	SkeletonReader->SetReadOnly(true);
 	
 	if (USkeletalMeshEditorContextObject* ContextObject = UE::SkeletalMeshEditorUtils::GetEditorContextObject(GetInteractiveToolsContext()))
diff --git a/Engine/Plugins/Animation/SkeletalMeshModelingTools/Source/SkeletalMeshModelingTools/Private/SkeletalMeshModelingToolsEditorMode.h b/Engine/Plugins/Animation/SkeletalMeshModelingTools/Source/SkeletalMeshModelingTools/Private/SkeletalMeshModelingToolsEditorMode.h
--- a/Engine/Plugins/Animation/SkeletalMeshModelingTools/Source/SkeletalMeshModelingTools/Private/SkeletalMeshModelingToolsEditorMode.h
+++ b/Engine/Plugins/Animation/SkeletalMeshModelingTools/Source/SkeletalMeshModelingTools/Private/SkeletalMeshModelingToolsEditorMode.h
@@ -212,8 +212,21 @@ private:
 	TObjectPtr<USkeletonModifier> SkeletonReader = nullptr;
 
 	TValueWatcher<int32> CacheChangeCountWatcher;
-	TValueWatcher<bool> CacheDynamicMeshSkeletonStatusWatcher;
+	
+	struct FSkeletonTreeHostState
+	{
+		bool bEnableDynamicMeshSkeleton = false;
+		bool bIsUsingToolSkeletonWidget = false;
+		bool bActiveToolModifiesSkeleton = false;
+		friend bool operator==(const FSkeletonTreeHostState&, const FSkeletonTreeHostState&) = default;
+	};
+	
+	// Decides which skeleton tree widget to show to the user
+	TValueWatcher<FSkeletonTreeHostState> SkeletonTreeHostStateWatcher;
 	TValueWatcher<int32> CacheSkeletonChangeCountWatcher;
+
+	bool bIsUsingToolSkeletonWidget = false;
+	bool bActiveToolModifiesSkeleton = false;
 	
 	TWeakPtr<FSkeletalMeshModelingToolsEditorModeToolkit> TypedToolkit = nullptr;
 	
diff --git a/Engine/Plugins/Animation/SkeletalMeshModelingTools/Source/SkeletalMeshModelingTools/Private/SkeletalMeshModelingToolsEditorModeToolkit.cpp b/Engine/Plugins/Animation/SkeletalMeshModelingTools/Source/SkeletalMeshModelingTools/Private/SkeletalMeshModelingToolsEditorModeToolkit.cpp
--- a/Engine/Plugins/Animation/SkeletalMeshModelingTools/Source/SkeletalMeshModelingTools/Private/SkeletalMeshModelingToolsEditorModeToolkit.cpp
+++ b/Engine/Plugins/Animation/SkeletalMeshModelingTools/Source/SkeletalMeshModelingTools/Private/SkeletalMeshModelingToolsEditorModeToolkit.cpp
@@ -550,6 +550,11 @@ void FSkeletalMeshModelingToolsEditorModeToolkit::OnEditorModeInitialized()
 			.Modifier(EditorMode->GetSkeletonReader())
 	];
 
+
+	SAssignNew(ToolSkeletonTreeHost, SBorder)
+		.Padding(4.f)
+		.BorderImage(FAppStyle::GetBrush("ToolPanel.GroupBorder"));
+	
 	// Save the default Skeleton Tree
 	if (ensure(ModeUILayer.IsValid()))
 	{
@@ -565,6 +570,12 @@ void FSkeletalMeshModelingToolsEditorModeToolkit::OnEditorModeInitialized()
 		}
 		DefaultSkeletonWidget = SkeletonTab->GetContent();
 	}
+
+	// Save the default Skeleton bone manipulation commands
+	const TSharedRef<FUICommandList>& CommandList = EditorMode->GetEditor()->GetToolkitCommands();
+	DefaultResetBoneTransformsAction = *(CommandList->GetActionForCommand(EditorMode->GetEditor()->GetResetBoneTransformsCommand()));
+	DefaultResetAllBoneTransformsAction = *(CommandList->GetActionForCommand(EditorMode->GetEditor()->GetResetAllBonesTransformsCommand()));
+
 	
 
 	SMorphTargetManager::FMorphTargetManagerDelegates Delegates;
@@ -617,17 +628,7 @@ void FSkeletalMeshModelingToolsEditorModeToolkit::ShowDynamicMeshSkeletonTree()
 		DynamicMeshSkeletonTree->GetNotifier()->HandleNotification({}, ESkeletalMeshNotifyType::HierarchyChanged);
 		DynamicMeshSkeletonTree->GetNotifier()->HandleNotification(SelectedBones, ESkeletalMeshNotifyType::BonesSelected);
 
-		const TSharedRef<FUICommandList>& CommandList = EditorMode->GetEditor()->GetToolkitCommands();
-		DefaultResetBoneTransformsAction = *(CommandList->GetActionForCommand(EditorMode->GetEditor()->GetResetBoneTransformsCommand()));
-		DefaultResetAllBoneTransformsAction = *(CommandList->GetActionForCommand(EditorMode->GetEditor()->GetResetAllBonesTransformsCommand()));
-
-		// Reset transforms
-		static constexpr bool bSelectedOnly = true;
-		CommandList->MapAction(EditorMode->GetEditor()->GetResetBoneTransformsCommand(),
-			FExecuteAction::CreateSP(this, &FSkeletalMeshModelingToolsEditorModeToolkit::ResetDynamicMeshBoneTransforms, bSelectedOnly));
 
-		CommandList->MapAction(EditorMode->GetEditor()->GetResetAllBonesTransformsCommand(),
-			FExecuteAction::CreateSP(this, &FSkeletalMeshModelingToolsEditorModeToolkit::ResetDynamicMeshBoneTransforms, !bSelectedOnly));
 	}
 }
 
@@ -648,66 +649,77 @@ void FSkeletalMeshModelingToolsEditorModeToolkit::ShowSkeletalMeshSkeletonTree()
 		SkeletonTab->SetContent(DefaultSkeletonWidget.ToSharedRef());
 
 		DynamicMeshSkeletonTreeNotifierBindScope.Reset();
+	}	
+}
+
+void FSkeletalMeshModelingToolsEditorModeToolkit::ShowToolSkeletonTree()
+{
+	if (ModeUILayer.IsValid())
+	{
+		TSharedPtr<FAssetEditorModeUILayer> ModeUILayerPtr = ModeUILayer.Pin();
+		TSharedPtr<FTabManager> TabManager = ModeUILayerPtr->GetTabManager();
+
+		const FName& SkeletonTreeViewID = FPersonaTabs::SkeletonTreeViewID;
+		TSharedPtr<SDockTab> SkeletonTab = TabManager->FindExistingLiveTab(SkeletonTreeViewID);
+		if (!SkeletonTab.IsValid())
+		{
+			SkeletonTab = TabManager->TryInvokeTab(SkeletonTreeViewID);
+		}
 
-		// Rebind commands if it was replaced before
-		const TSharedRef<FUICommandList>& CommandList = EditorMode->GetEditor()->GetToolkitCommands();
+		SkeletonTab->SetContent(ToolSkeletonTreeHost.ToSharedRef());
+	}		
+}
 
-		// Reset transforms
-		CommandList->MapAction(EditorMode->GetEditor()->GetResetBoneTransformsCommand(), DefaultResetBoneTransformsAction);
-		CommandList->MapAction(EditorMode->GetEditor()->GetResetAllBonesTransformsCommand(), DefaultResetAllBoneTransformsAction);
-	}	
+void FSkeletalMeshModelingToolsEditorModeToolkit::EnableDynamicMeshSkeletonCommands()
+{
+	const TSharedRef<FUICommandList>& CommandList = EditorMode->GetEditor()->GetToolkitCommands();
+
+	// Reset transforms
+	static constexpr bool bSelectedOnly = true;
+	CommandList->MapAction(EditorMode->GetEditor()->GetResetBoneTransformsCommand(),
+		FExecuteAction::CreateSP(this, &FSkeletalMeshModelingToolsEditorModeToolkit::ResetDynamicMeshBoneTransforms, bSelectedOnly));
+
+	CommandList->MapAction(EditorMode->GetEditor()->GetResetAllBonesTransformsCommand(),
+		FExecuteAction::CreateSP(this, &FSkeletalMeshModelingToolsEditorModeToolkit::ResetDynamicMeshBoneTransforms, !bSelectedOnly));
 }
 
-bool FSkeletalMeshModelingToolsEditorModeToolkit::BindToolSkeletonTree(ISkeletalMeshEditingInterface* InEditingInterface)
+void FSkeletalMeshModelingToolsEditorModeToolkit::EnableSkeletalMeshSkeletonCommands()
 {
-	bWasUsingDynamicMeshSkeletonBeforeToolActivation = EditorMode->GetCurrentEditingCache()->IsDynamicMeshSkeletonEnabled();
+	// Rebind commands if it was replaced before
+	const TSharedRef<FUICommandList>& CommandList = EditorMode->GetEditor()->GetToolkitCommands();
 
+	// Reset transforms
+	CommandList->MapAction(EditorMode->GetEditor()->GetResetBoneTransformsCommand(), DefaultResetBoneTransformsAction);
+	CommandList->MapAction(EditorMode->GetEditor()->GetResetAllBonesTransformsCommand(), DefaultResetAllBoneTransformsAction);	
+}
+
+void FSkeletalMeshModelingToolsEditorModeToolkit::DisableEditingCacheSkeletonCommands()
+{
+	// Rebind commands if it was replaced before
+	const TSharedRef<FUICommandList>& CommandList = EditorMode->GetEditor()->GetToolkitCommands();
+
+	// Reset transforms
+	CommandList->UnmapAction(EditorMode->GetEditor()->GetResetBoneTransformsCommand());
+	CommandList->UnmapAction(EditorMode->GetEditor()->GetResetAllBonesTransformsCommand());	
+}
+
+bool FSkeletalMeshModelingToolsEditorModeToolkit::BindToolSkeletonTree(ISkeletalMeshEditingInterface* InEditingInterface)
+{
 	TSharedPtr<SWidget> ToolSkeletonTree = InEditingInterface->GetCustomSkeletonTreeWidget(EditorMode->GetModeBinding());
 
 	if (!ToolSkeletonTree.IsValid())
 	{
 		return false;
 	}
-	
-	if (!ModeUILayer.IsValid())
-	{
-		return false;
-	}
-	
-	TSharedPtr<FAssetEditorModeUILayer> ModeUILayerPtr = ModeUILayer.Pin();
-	TSharedPtr<FTabManager> TabManager = ModeUILayerPtr->GetTabManager();
-	
-	const FName& SkeletonTreeId = FPersonaTabs::SkeletonTreeViewID;
-	TSharedPtr<SDockTab> SkeletonTab = TabManager->FindExistingLiveTab(SkeletonTreeId);
-	if (!SkeletonTab.IsValid())
-	{
-		SkeletonTab = TabManager->TryInvokeTab(SkeletonTreeId);
-	}
-
-	SAssignNew(ToolSkeletonHost, SBorder)
-		.Padding(4.f)
-		.BorderImage(FAppStyle::GetBrush("ToolPanel.GroupBorder"))
-	[
-		ToolSkeletonTree.ToSharedRef()	
-	];
-
-	SkeletonTab->SetContent(ToolSkeletonHost.ToSharedRef());
 
+	ToolSkeletonTreeHost->SetContent(ToolSkeletonTree.ToSharedRef());
+	
 	return true;
 }
 
 void FSkeletalMeshModelingToolsEditorModeToolkit::UnbindToolSkeletonTree()
 {
-	if (bWasUsingDynamicMeshSkeletonBeforeToolActivation)
-	{
-		ShowDynamicMeshSkeletonTree();
-	}
-	else
-	{
-		ShowSkeletalMeshSkeletonTree();
-	}
-
-	ToolSkeletonHost.Reset();
+	ToolSkeletonTreeHost->SetContent(SNullWidget::NullWidget);
 }
 
 void FSkeletalMeshModelingToolsEditorModeToolkit::RefreshMorphTargetManager()
diff --git a/Engine/Plugins/Animation/SkeletalMeshModelingTools/Source/SkeletalMeshModelingTools/Private/SkeletalMeshModelingToolsEditorModeToolkit.h b/Engine/Plugins/Animation/SkeletalMeshModelingTools/Source/SkeletalMeshModelingTools/Private/SkeletalMeshModelingToolsEditorModeToolkit.h
--- a/Engine/Plugins/Animation/SkeletalMeshModelingTools/Source/SkeletalMeshModelingTools/Private/SkeletalMeshModelingToolsEditorModeToolkit.h
+++ b/Engine/Plugins/Animation/SkeletalMeshModelingTools/Source/SkeletalMeshModelingTools/Private/SkeletalMeshModelingToolsEditorModeToolkit.h
@@ -40,10 +40,15 @@ public:
 	void ResetDynamicMeshBoneTransforms(bool bSelectedOnly);
 	void ShowDynamicMeshSkeletonTree();
 	void ShowSkeletalMeshSkeletonTree();
+	void ShowToolSkeletonTree();
+	
+	void EnableDynamicMeshSkeletonCommands();
+	void EnableSkeletalMeshSkeletonCommands();
+	// When tool has its own skeleton, the editing cache skeleton is hidden and shouldn't receive commands like reset bone transforms
+	void DisableEditingCacheSkeletonCommands();
 
 	bool BindToolSkeletonTree(ISkeletalMeshEditingInterface* InTool);
 	void UnbindToolSkeletonTree();
-
 	
 	void RefreshMorphTargetManager();
 private:
@@ -103,8 +108,7 @@ private:
 	TSharedPtr<SBorder> DynamicMeshSkeletonWidget;
 	TSharedPtr<SReferenceSkeletonTree> DynamicMeshSkeletonTree;
 
-	bool bWasUsingDynamicMeshSkeletonBeforeToolActivation = false;
-	TSharedPtr<SBorder> ToolSkeletonHost;
+	TSharedPtr<SBorder> ToolSkeletonTreeHost;
 	
 
 	TSharedPtr<SWidget> DefaultMorphTargetWidget; 
diff --git a/Engine/Plugins/Runtime/MeshModelingToolset/Source/MeshModelingToolsEditorOnly/Private/SkeletalMesh/SReferenceSkeletonTree.cpp b/Engine/Plugins/Runtime/MeshModelingToolset/Source/MeshModelingToolsEditorOnly/Private/SkeletalMesh/SReferenceSkeletonTree.cpp
--- a/Engine/Plugins/Runtime/MeshModelingToolset/Source/MeshModelingToolsEditorOnly/Private/SkeletalMesh/SReferenceSkeletonTree.cpp
+++ b/Engine/Plugins/Runtime/MeshModelingToolset/Source/MeshModelingToolsEditorOnly/Private/SkeletalMesh/SReferenceSkeletonTree.cpp
@@ -298,7 +298,10 @@ SReferenceSkeletonTree::SReferenceSkeletonTree()
 
 SReferenceSkeletonTree::~SReferenceSkeletonTree()
 {
-	// GEditor->UnregisterForUndo(this);
+	if (GEditor != nullptr)
+	{
+		GEditor->UnregisterForUndo(this);
+	}
 }
 
 void SReferenceSkeletonTree::Construct(const FArguments& InArgs)
diff --git a/Engine/Plugins/Runtime/MeshModelingToolset/Source/MeshModelingToolsEditorOnly/Private/SkeletalMesh/SkinWeightsPaintTool.cpp b/Engine/Plugins/Runtime/MeshModelingToolset/Source/MeshModelingToolsEditorOnly/Private/SkeletalMesh/SkinWeightsPaintTool.cpp
--- a/Engine/Plugins/Runtime/MeshModelingToolset/Source/MeshModelingToolsEditorOnly/Private/SkeletalMesh/SkinWeightsPaintTool.cpp
+++ b/Engine/Plugins/Runtime/MeshModelingToolset/Source/MeshModelingToolsEditorOnly/Private/SkeletalMesh/SkinWeightsPaintTool.cpp
@@ -1946,7 +1946,9 @@ void USkinWeightsPaintTool::Setup()
 	auto OnSelectionChangedLambda = [this](){OnSelectionChanged.Broadcast();};
 	MeshSelector->InitialSetup(TargetWorld.Get(), this, OnSelectionChangedLambda);
 
-	InitializeSkeletonReader();
+	SkeletonReader = NewObject<USkeletonModifier>();
+	SkeletonReader->SetReferenceSkeleton(RefSkeleton);
+	SkeletonReader->SetReadOnly(true);
 	
 	// run all initialization for mesh/weights
 	UpdateCurrentlyEditedMesh(EditedMesh);
@@ -4321,22 +4323,6 @@ void USkinWeightsPaintTool::ToggleLockBone(BoneIndex InBoneIndex, bool bLock)
 	}	
 }
 
-void USkinWeightsPaintTool::InitializeSkeletonReader()
-{
-	if (ISkeletalMeshBackedTarget* SkeletalMeshBackedTarget = Cast<ISkeletalMeshBackedTarget>(Target))
-	{
-		SkeletonReader = NewObject<USkeletonModifier>();
-		SkeletonReader->SetSkeletalMesh(SkeletalMeshBackedTarget->GetSkeletalMesh());
-		SkeletonReader->SetReadOnly(true);
-	}
-	else if (IPersistentDynamicMeshSource* DynamicMeshSource = Cast<IPersistentDynamicMeshSource>(Target))
-	{
-		SkeletonReader = NewObject<USkeletonModifier>();
-		SkeletonReader->SetDynamicMesh(DynamicMeshSource->GetDynamicMeshContainer());
-		SkeletonReader->SetReadOnly(true);
-	}
-}
-
 bool USkinWeightsPaintTool::ShouldProceedWithAllBonesUnlocked()
 {
 	if (GetFinalLockedBoneIndices().IsEmpty())
diff --git a/Engine/Plugins/Runtime/MeshModelingToolset/Source/MeshModelingToolsEditorOnly/Public/SkeletalMesh/SkinWeightsPaintTool.h b/Engine/Plugins/Runtime/MeshModelingToolset/Source/MeshModelingToolsEditorOnly/Public/SkeletalMesh/SkinWeightsPaintTool.h
--- a/Engine/Plugins/Runtime/MeshModelingToolset/Source/MeshModelingToolsEditorOnly/Public/SkeletalMesh/SkinWeightsPaintTool.h
+++ b/Engine/Plugins/Runtime/MeshModelingToolset/Source/MeshModelingToolsEditorOnly/Public/SkeletalMesh/SkinWeightsPaintTool.h
@@ -35,6 +35,7 @@ class FEditorViewportClient;
 class UVolumetricBrushStampIndicator;
 class USkeletonModifier;
 class SButton;
+class FSkeletalMeshNotifierBindScope;
 
 namespace UE::Geometry 
 {
@@ -920,7 +921,6 @@ protected:
 	UE_API void IsolateSelectedBones();
 	UE_API void ToggleBypassLocks();
 	UE_API void ToggleLockBone(BoneIndex InBoneIndex, bool bLock);
-	UE_API void InitializeSkeletonReader();
 
 	TMap<BoneIndex, TSharedPtr<SButton>> LockButtons;
 	TSet<BoneIndex> LockedBoneIndices;
diff --git a/Engine/Plugins/Runtime/MeshModelingToolset/Source/SkeletalMeshModifiers/Private/SkeletonModifier.cpp b/Engine/Plugins/Runtime/MeshModelingToolset/Source/SkeletalMeshModifiers/Private/SkeletonModifier.cpp
--- a/Engine/Plugins/Runtime/MeshModelingToolset/Source/SkeletalMeshModifiers/Private/SkeletonModifier.cpp
+++ b/Engine/Plugins/Runtime/MeshModelingToolset/Source/SkeletalMeshModifiers/Private/SkeletonModifier.cpp
@@ -416,18 +416,8 @@ bool USkeletonModifier::SetSkeletalMesh(USkeletalMesh* InSkeletalMesh)
 	}
 
 	// store reference skeleton to edit
-	ReferenceSkeleton = MakeUnique<FReferenceSkeleton>();
-	*ReferenceSkeleton = SkeletalMesh->GetRefSkeleton();
-
-	TransformComposer.Reset(new FTransformComposer(*ReferenceSkeleton));
-
-	// store initial bones indices to track for changes
-	const int32 NumBones = ReferenceSkeleton->GetRawBoneNum();
-	BoneIndexTracker.Reserve(NumBones);
-	for (int32 Index = 0; Index < NumBones; Index++)
-	{
-		BoneIndexTracker.Add(Index);
-	}
+	SetReferenceSkeleton(SkeletalMesh->GetRefSkeleton());
+	
 	return true;
 	
 #else
@@ -438,6 +428,11 @@ bool USkeletonModifier::SetSkeletalMesh(USkeletalMesh* InSkeletalMesh)
 
 bool USkeletonModifier::SetDynamicMesh(UDynamicMesh* InDynamicMesh)
 {
+	WeakDynamicMesh = nullptr;
+	ReferenceSkeleton.Reset();
+	TransformComposer.Reset();
+	BoneIndexTracker.Reset();
+	
 #if WITH_EDITORONLY_DATA
 	
 	if (!InDynamicMesh)
@@ -446,18 +441,12 @@ bool USkeletonModifier::SetDynamicMesh(UDynamicMesh* InDynamicMesh)
 	}
 	
 	WeakDynamicMesh = InDynamicMesh;
-
-	ReferenceSkeleton.Reset();
-	TransformComposer.Reset();
-	BoneIndexTracker.Reset();
-
-	// store reference skeleton to edit
-	ReferenceSkeleton = MakeUnique<FReferenceSkeleton>();
+	FReferenceSkeleton MeshSkeleton;
 	{
 		const FDynamicMesh3& Mesh = InDynamicMesh->GetMeshRef();
 		
 		const int32 NumBones = Mesh.Attributes()->GetNumBones();
-		FReferenceSkeletonModifier RefSkeletonModifier(*ReferenceSkeleton, nullptr);
+		FReferenceSkeletonModifier RefSkeletonModifier(MeshSkeleton, nullptr);
 
 		using namespace UE::Geometry;
 		const FDynamicMeshBoneNameAttribute* BoneNameAttrib = Mesh.Attributes()->GetBoneNames();
@@ -473,9 +462,29 @@ bool USkeletonModifier::SetDynamicMesh(UDynamicMesh* InDynamicMesh)
 			FTransform Transform = BonePoses->GetValue(BoneIndex);
 			RefSkeletonModifier.Add(BoneInfo, Transform);
 		}	
-	}
+	}	
+		
+	SetReferenceSkeleton(MeshSkeleton);
+	
+
+	return true;
 	
+#else
+	ensureMsgf(false, TEXT("Skeleton Modifier is an editor only feature."));
+	return false;
+#endif
+
+}
+
+bool USkeletonModifier::SetReferenceSkeleton(const FReferenceSkeleton& InReferenceSkeleton)
+{
+	ReferenceSkeleton.Reset();
+	TransformComposer.Reset();
+	BoneIndexTracker.Reset();
 
+	// store reference skeleton to read from
+	ReferenceSkeleton = MakeUnique<FReferenceSkeleton>();
+	*ReferenceSkeleton = InReferenceSkeleton;
 	TransformComposer.Reset(new FTransformComposer(*ReferenceSkeleton));
 
 	// store initial bones indices to track for changes
@@ -485,13 +494,8 @@ bool USkeletonModifier::SetDynamicMesh(UDynamicMesh* InDynamicMesh)
 	{
 		BoneIndexTracker.Add(Index);
 	}
-	return true;
-	
-#else
-	ensureMsgf(false, TEXT("Skeleton Modifier is an editor only feature."));
-	return false;
-#endif
 
+	return true;
 }
 
 void USkeletonModifier::SetReadOnly(bool bInReadOnly)
diff --git a/Engine/Plugins/Runtime/MeshModelingToolset/Source/SkeletalMeshModifiers/Public/SkeletonModifier.h b/Engine/Plugins/Runtime/MeshModelingToolset/Source/SkeletalMeshModifiers/Public/SkeletonModifier.h
--- a/Engine/Plugins/Runtime/MeshModelingToolset/Source/SkeletalMeshModifiers/Public/SkeletonModifier.h
+++ b/Engine/Plugins/Runtime/MeshModelingToolset/Source/SkeletalMeshModifiers/Public/SkeletonModifier.h
@@ -142,6 +142,8 @@ public:
 
 
 	UE_API bool SetDynamicMesh(UDynamicMesh* InDynamicMesh);
+
+	UE_API bool SetReferenceSkeleton(const FReferenceSkeleton& InReferenceSkeleton);
 	
 	UE_API void SetReadOnly(bool bInReadOnly);
 	UE_API bool IsReadOnly() const;
