[SkeletalMesh ] Added a delegate to notify modeling tools when the mesh is about to change so that it has a chance to apply pending changes to the asset

diff --git a/Engine/Plugins/Animation/SkeletalMeshModelingTools/Source/SkeletalMeshModelingTools/Private/SkeletalMeshModelingToolsEditorMode.cpp b/Engine/Plugins/Animation/SkeletalMeshModelingTools/Source/SkeletalMeshModelingTools/Private/SkeletalMeshModelingToolsEditorMode.cpp
--- a/Engine/Plugins/Animation/SkeletalMeshModelingTools/Source/SkeletalMeshModelingTools/Private/SkeletalMeshModelingToolsEditorMode.cpp
+++ b/Engine/Plugins/Animation/SkeletalMeshModelingTools/Source/SkeletalMeshModelingTools/Private/SkeletalMeshModelingToolsEditorMode.cpp
@@ -1114,6 +1114,7 @@ void USkeletalMeshModelingToolsEditorMode::SetEditorBinding(const TWeakPtr<ISkel
 	Editor.Pin()->OnPreSaveAssetAs().AddUObject(this, &USkeletalMeshModelingToolsEditorMode::ApplyChanges);
 
 	SkeletalMesh = InSkeletalMeshEditor.Pin()->GetPersonaToolkit()->GetMesh();
+	SkeletalMesh->GetOnPreMeshChange().AddUObject(this, &USkeletalMeshModelingToolsEditorMode::HandleSkeletalMeshPreChange);
 	SkeletalMesh->GetOnMeshChanged().AddUObject(this, &USkeletalMeshModelingToolsEditorMode::HandleSkeletalMeshChanged);
 	SkeletonReader = NewObject<USkeletonModifier>();
 	SkeletonReader->SetReferenceSkeleton(SkeletalMesh->GetRefSkeleton());
@@ -1152,6 +1153,11 @@ USkeletalMesh* USkeletalMeshModelingToolsEditorMode::GetSkeletalMesh() const
 	return SkeletalMesh.Get();
 }
 
+void USkeletalMeshModelingToolsEditorMode::HandleSkeletalMeshPreChange()
+{
+	GetCurrentEditingCache()->ApplyChanges();
+}
+
 void USkeletalMeshModelingToolsEditorMode::HandleSkeletalMeshChanged()
 {
 	if (GetCurrentEditingCache())
@@ -1169,8 +1175,6 @@ void USkeletalMeshModelingToolsEditorMode::HandleSkeletalMeshChanged()
 			RecreateEditingCache(GetEditingLOD());
 		}
 	}
-	
-
 }
 
 bool USkeletalMeshModelingToolsEditorMode::CanSetEditingLOD()
diff --git a/Engine/Plugins/Animation/SkeletalMeshModelingTools/Source/SkeletalMeshModelingTools/Private/SkeletalMeshModelingToolsEditorMode.h b/Engine/Plugins/Animation/SkeletalMeshModelingTools/Source/SkeletalMeshModelingTools/Private/SkeletalMeshModelingToolsEditorMode.h
--- a/Engine/Plugins/Animation/SkeletalMeshModelingTools/Source/SkeletalMeshModelingTools/Private/SkeletalMeshModelingToolsEditorMode.h
+++ b/Engine/Plugins/Animation/SkeletalMeshModelingTools/Source/SkeletalMeshModelingTools/Private/SkeletalMeshModelingToolsEditorMode.h
@@ -115,6 +115,7 @@ public:
 	TSharedPtr<ISkeletalMeshEditor> GetEditor();
 	
 	USkeletalMesh* GetSkeletalMesh() const;
+	void HandleSkeletalMeshPreChange();
 	void HandleSkeletalMeshChanged();
 	
 	bool CanSetEditingLOD();
diff --git a/Engine/Source/Runtime/Engine/Classes/Engine/SkeletalMesh.h b/Engine/Source/Runtime/Engine/Classes/Engine/SkeletalMesh.h
--- a/Engine/Source/Runtime/Engine/Classes/Engine/SkeletalMesh.h
+++ b/Engine/Source/Runtime/Engine/Classes/Engine/SkeletalMesh.h
@@ -445,6 +445,8 @@ class USkeletalMesh : public USkinnedAsset, public IInterface_CollisionDataProvi
 	ENGINE_API virtual ~USkeletalMesh() override;
 
 #if WITH_EDITOR
+	/** Notification when something is about to be changed */
+	DECLARE_MULTICAST_DELEGATE(FOnPreMeshChange);
 	/** Notification when anything changed */
 	DECLARE_MULTICAST_DELEGATE(FOnMeshChanged);
 #endif
@@ -2383,7 +2385,9 @@ public:
 		SamplingInfo = InSamplingInfo;
 		PRAGMA_ENABLE_DEPRECATION_WARNINGS
 	}
-
+	
+	const FOnPreMeshChange& GetOnPreMeshChange() const { return OnPreMeshChange; }
+	FOnPreMeshChange& GetOnPreMeshChange() { return OnPreMeshChange; }
 	const FOnMeshChanged& GetOnMeshChanged() const { return OnMeshChanged; }
 	FOnMeshChanged& GetOnMeshChanged() { return OnMeshChanged; }
 #endif
@@ -2424,6 +2428,7 @@ protected:
 	}
 #endif
 #if WITH_EDITOR
+	FOnPreMeshChange OnPreMeshChange;
 	FOnMeshChanged OnMeshChanged;
 #endif
 
diff --git a/Engine/Source/Runtime/Engine/Private/SkeletalMesh.cpp b/Engine/Source/Runtime/Engine/Private/SkeletalMesh.cpp
--- a/Engine/Source/Runtime/Engine/Private/SkeletalMesh.cpp
+++ b/Engine/Source/Runtime/Engine/Private/SkeletalMesh.cpp
@@ -328,6 +328,7 @@ void FScopedSkeletalMeshPostEditChange::SetSkeletalMesh(USkeletalMesh* InSkeleta
 
 		if (bCallPostEditChange)
 		{
+			SkeletalMesh->GetOnPreMeshChange().Broadcast();
 			//Make sure the render ressource use by the skeletalMesh is free, we will reconstruct them when a PostEditChange will be call
 			SkeletalMesh->FlushRenderState();
 		}
