[SKM ModelingTools] Minor fixes

#rnx

diff --git a/Engine/Plugins/Animation/SkeletalMeshModelingTools/Source/SkeletalMeshModelingTools/Private/SkeletalMeshModelingToolsEditorMode.cpp b/Engine/Plugins/Animation/SkeletalMeshModelingTools/Source/SkeletalMeshModelingTools/Private/SkeletalMeshModelingToolsEditorMode.cpp
--- a/Engine/Plugins/Animation/SkeletalMeshModelingTools/Source/SkeletalMeshModelingTools/Private/SkeletalMeshModelingToolsEditorMode.cpp
+++ b/Engine/Plugins/Animation/SkeletalMeshModelingTools/Source/SkeletalMeshModelingTools/Private/SkeletalMeshModelingToolsEditorMode.cpp
@@ -613,6 +613,12 @@ void USkeletalMeshModelingToolsEditorMode::Exit()
 		Editor.Pin()->OnPreSaveAsset().RemoveAll(this);
 		Editor.Pin()->OnPreSaveAssetAs().RemoveAll(this);
 	}
+
+	if (SkeletalMesh.IsValid())
+	{
+		SkeletalMesh->GetOnPreMeshChange().RemoveAll(this);
+		SkeletalMesh->GetOnMeshChanged().RemoveAll(this);
+	}
 	
 #if ENABLE_STYLUS_SUPPORT
 	StylusStateTracker = nullptr;
@@ -1155,7 +1161,10 @@ USkeletalMesh* USkeletalMeshModelingToolsEditorMode::GetSkeletalMesh() const
 
 void USkeletalMeshModelingToolsEditorMode::HandleSkeletalMeshPreChange()
 {
-	GetCurrentEditingCache()->ApplyChanges();
+	if (GetCurrentEditingCache())
+	{
+		GetCurrentEditingCache()->ApplyChanges();
+	}
 }
 
 void USkeletalMeshModelingToolsEditorMode::HandleSkeletalMeshChanged()
