[SKM Modelling Tools] avoid transform vert color to srgb in the first place since vert color is not transform back to linear when commiting the dyna mesh to mesh description by default

diff --git a/Engine/Plugins/Animation/SkeletalMeshModelingTools/Source/SkeletalMeshModelingTools/Private/Components/SKMBackedDynaMeshComponent.cpp b/Engine/Plugins/Animation/SkeletalMeshModelingTools/Source/SkeletalMeshModelingTools/Private/Components/SKMBackedDynaMeshComponent.cpp
--- a/Engine/Plugins/Animation/SkeletalMeshModelingTools/Source/SkeletalMeshModelingTools/Private/Components/SKMBackedDynaMeshComponent.cpp
+++ b/Engine/Plugins/Animation/SkeletalMeshModelingTools/Source/SkeletalMeshModelingTools/Private/Components/SKMBackedDynaMeshComponent.cpp
@@ -194,6 +194,9 @@ EMeshLODIdentifier USkeletalMeshBackedDynamicMeshComponent::Init(USkeletalMesh*
 	FMeshDescription* MeshDescription = InSkeletalMesh->GetMeshDescription(LODIndex);
 	FDynamicMesh3 LODMesh;
 	FMeshDescriptionToDynamicMesh Converter;
+	// By default, when converting dyna mesh back to mesh description, we don't transform vert color, so avoid doing it in the first place.
+	// See FDynamicMeshCommitInfo
+	Converter.bTransformVertexColorsLinearToSRGB = false;
 	constexpr bool bWantTangents = true;
 	Converter.bVIDsFromNonManifoldMeshDescriptionAttr = true;
 	Converter.Convert(MeshDescription, LODMesh, bWantTangents);
