[Gizmos] Remove usage of custom material for the SkeletalMeshGizmoMaterial that was no longer functioning correctly.

#rnx

diff --git a/Engine/Plugins/Animation/SkeletalMeshModelingTools/Source/SkeletalMeshModelingTools/Private/SkeletalMeshGizmoUtils.cpp b/Engine/Plugins/Animation/SkeletalMeshModelingTools/Source/SkeletalMeshModelingTools/Private/SkeletalMeshGizmoUtils.cpp
--- a/Engine/Plugins/Animation/SkeletalMeshModelingTools/Source/SkeletalMeshModelingTools/Private/SkeletalMeshGizmoUtils.cpp
+++ b/Engine/Plugins/Animation/SkeletalMeshModelingTools/Source/SkeletalMeshModelingTools/Private/SkeletalMeshGizmoUtils.cpp
@@ -167,13 +167,7 @@ USkeletalMeshGizmoWrapperBase* USkeletalMeshGizmoContextObject::GetNewWrapper(UI
 
 const FGizmoCustomization& USkeletalMeshGizmoContextObject::GetGizmoCustomization()
 {
-	static const FString MaterialName = TEXT("/SkeletalMeshModelingTools/SkeletalMeshGizmoMaterial.SkeletalMeshGizmoMaterial");
-	UMaterial* Material = LoadObject<UMaterial>(nullptr, *MaterialName, nullptr, LOAD_None, nullptr);
-
-	check(Material);
-	
 	static FGizmoCustomization Customization;
-		Customization.Material = Material;
 		Customization.SizeCoefficient = 1.5f;
 
 	return Customization;
