[SKM Modeling Tools] Fixed bind skin tool not showing the skeleton tree

diff --git a/Engine/Plugins/Animation/SkeletalMeshModelingTools/Source/SkeletalMeshModelingTools/Private/SkeletalMeshModelingToolsEditorMode.cpp b/Engine/Plugins/Animation/SkeletalMeshModelingTools/Source/SkeletalMeshModelingTools/Private/SkeletalMeshModelingToolsEditorMode.cpp
--- a/Engine/Plugins/Animation/SkeletalMeshModelingTools/Source/SkeletalMeshModelingTools/Private/SkeletalMeshModelingToolsEditorMode.cpp
+++ b/Engine/Plugins/Animation/SkeletalMeshModelingTools/Source/SkeletalMeshModelingTools/Private/SkeletalMeshModelingToolsEditorMode.cpp
@@ -468,7 +468,6 @@ TSharedPtr<FTabManager> USkeletalMeshModelingToolsEditorMode::GetAssociatedTabMa
 
 void USkeletalMeshModelingToolsEditorMode::BindToolSkeletonTree(ISkeletalMeshEditingInterface* InEditingInterface)
 {
-	bIsUsingToolSkeletonWidget = true;
 	bActiveToolModifiesSkeleton = InEditingInterface->CanModifySkeleton();
 	if (InEditingInterface->CanModifySkeleton())
 	{
@@ -477,7 +476,7 @@ void USkeletalMeshModelingToolsEditorMode::BindToolSkeletonTree(ISkeletalMeshEdi
 		GetCurrentEditingCache()->GetNotifier()->HandleNotification({}, ESkeletalMeshNotifyType::BonesSelected);
 	}
 	
-	TypedToolkit.Pin()->BindToolSkeletonTree(InEditingInterface);
+	bIsUsingToolSkeletonWidget = TypedToolkit.Pin()->BindToolSkeletonTree(InEditingInterface);
 }
 
 void USkeletalMeshModelingToolsEditorMode::UnbindToolSkeletonTree()
