[SkeletalMeshModelingMode] Added more nullchecks

#Synced-CL 48863170

diff --git a/Engine/Plugins/Animation/SkeletalMeshModelingTools/Source/SkeletalMeshModelingTools/Private/SkeletalMeshModelingToolsEditorMode.cpp b/Engine/Plugins/Animation/SkeletalMeshModelingTools/Source/SkeletalMeshModelingTools/Private/SkeletalMeshModelingToolsEditorMode.cpp
--- a/Engine/Plugins/Animation/SkeletalMeshModelingTools/Source/SkeletalMeshModelingTools/Private/SkeletalMeshModelingToolsEditorMode.cpp
+++ b/Engine/Plugins/Animation/SkeletalMeshModelingTools/Source/SkeletalMeshModelingTools/Private/SkeletalMeshModelingToolsEditorMode.cpp
@@ -497,14 +497,13 @@ void USkeletalMeshModelingToolsEditorMode::UnbindToolSkeletonTree()
 	bActiveToolModifiesSkeleton = false;
 }
 
-const TArray<FTransform>& USkeletalMeshModelingToolsEditorMode::GetComponentSpaceBoneTransforms() const
-{
-	return GetCurrentEditingCache()->GetComponentSpaceBoneTransforms();
-}
 
 void USkeletalMeshModelingToolsEditorMode::ToggleBoneManipulation(bool bEnable)
 {
-	GetCurrentEditingCache()->ToggleBoneManipulation(bEnable);
+	if (GetCurrentEditingCache())
+	{
+		GetCurrentEditingCache()->ToggleBoneManipulation(bEnable);
+	}
 }
 
 USkeletonModifier* USkeletalMeshModelingToolsEditorMode::GetSkeletonReader()
@@ -525,6 +524,11 @@ TArray<FName> USkeletalMeshModelingToolsEditorMode::GetSelectedBones() const
 
 FName USkeletalMeshModelingToolsEditorMode::GetEditingMorphTarget()
 {
+	if (!GetCurrentEditingCache())
+	{
+		return NAME_None;
+	}
+	
 	if (!GetCurrentEditingCache()->GetMorphTargets().Contains(EditingMorphTarget))
 	{
 		return NAME_None;
@@ -546,11 +550,20 @@ void USkeletalMeshModelingToolsEditorMode::HandleSetEditingMorphTarget(FName InM
 
 TArray<FName> USkeletalMeshModelingToolsEditorMode::GetMorphTargets()
 {
+	if (!GetCurrentEditingCache())
+	{
+		return {};
+	}
 	return GetCurrentEditingCache()->GetMorphTargets();
 }
 
 float USkeletalMeshModelingToolsEditorMode::GetMorphTargetWeight(FName MorphTarget)
 {
+	if (!GetCurrentEditingCache())
+	{
+		return 0; 
+	}
+	
 	return GetCurrentEditingCache()->GetMorphTargetWeight(MorphTarget);
 }
 
@@ -576,11 +589,17 @@ TArray<FName> USkeletalMeshModelingToolsEditorMode::HandleDuplicateMorphTargets(
 
 void USkeletalMeshModelingToolsEditorMode::HandleSetMorphTargetWeight(FName MorphTarget, float Weight)
 {
+	
 	return GetCurrentEditingCache()->HandleSetMorphTargetWeight(MorphTarget, Weight);
 }
 
 bool USkeletalMeshModelingToolsEditorMode::GetMorphTargetAutoFill(FName MorphTarget)
 {
+	if (!GetCurrentEditingCache())
+	{
+		return false; 
+	}
+	
 	return GetCurrentEditingCache()->GetMorphTargetAutoFill(MorphTarget);
 }
 
@@ -592,9 +611,12 @@ void USkeletalMeshModelingToolsEditorMode::HandleSetMorphTargetAutoFill(FName Mo
 
 void USkeletalMeshModelingToolsEditorMode::Exit()
 {
-	GetCurrentEditingCache()->ApplyChanges();
-	GetCurrentEditingCache()->Destroy();
-	CurrentEditingCache = nullptr;
+	if (GetCurrentEditingCache())
+	{
+		GetCurrentEditingCache()->ApplyChanges();
+		GetCurrentEditingCache()->Destroy();
+		CurrentEditingCache = nullptr;
+	}
 	
 	UEditorInteractiveToolsContext* InteractiveToolsContext = GetInteractiveToolsContext();
 	UE::TransformGizmoUtil::DeregisterTransformGizmoContextObject(InteractiveToolsContext);
@@ -742,7 +764,7 @@ void USkeletalMeshModelingToolsEditorMode::Tick(FEditorViewportClient* ViewportC
 
 bool USkeletalMeshModelingToolsEditorMode::HandleClick(FEditorViewportClient* InViewportClient, HHitProxy *HitProxy, const FViewportClick &Click)
 {
-	if (GetCurrentEditingCache()->HandleClick(HitProxy))
+	if (GetCurrentEditingCache() && GetCurrentEditingCache()->HandleClick(HitProxy))
 	{
 		return true;
 	}
@@ -873,7 +895,7 @@ bool USkeletalMeshModelingToolsEditorMode::ShouldToolStartBeAllowed(const FStrin
 
 bool USkeletalMeshModelingToolsEditorMode::StartTracking(FEditorViewportClient* InViewportClient, FViewport* InViewport)
 {
-	if (GetCurrentEditingCache()->IsDynamicMeshBoneManipulationEnabled())
+	if (GetCurrentEditingCache() && GetCurrentEditingCache()->IsDynamicMeshBoneManipulationEnabled())
 	{
 		if (GetCurrentEditingCache()->GetFirstSelectedBoneIndex() != INDEX_NONE)
 		{
@@ -899,7 +921,7 @@ bool USkeletalMeshModelingToolsEditorMode::StartTracking(FEditorViewportClient*
 
 bool USkeletalMeshModelingToolsEditorMode::EndTracking(FEditorViewportClient* InViewportClient, FViewport* InViewport)
 {
-	if (GetCurrentEditingCache()->GetSkeletonPoser()->IsRecordingPoseChange())
+	if (GetCurrentEditingCache() && GetCurrentEditingCache()->GetSkeletonPoser()->IsRecordingPoseChange())
 	{
 		GetCurrentEditingCache()->GetSkeletonPoser()->EndPoseChange();
 		
@@ -912,7 +934,7 @@ bool USkeletalMeshModelingToolsEditorMode::EndTracking(FEditorViewportClient* In
 bool USkeletalMeshModelingToolsEditorMode::InputDelta(FEditorViewportClient* InViewportClient, FViewport* InViewport, FVector& InDrag,
 	FRotator& InRot, FVector& InScale)
 {
-	if (GetCurrentEditingCache()->GetSkeletonPoser()->IsRecordingPoseChange())
+	if (GetCurrentEditingCache() && GetCurrentEditingCache()->GetSkeletonPoser()->IsRecordingPoseChange())
 	{
 		// Actually change bone transform
 		const EAxisList::Type CurrentAxis = InViewportClient->GetCurrentWidgetAxis();
@@ -985,7 +1007,7 @@ void USkeletalMeshModelingToolsEditorMode::DrawHUD(FEditorViewportClient* Viewpo
 
 bool USkeletalMeshModelingToolsEditorMode::AllowWidgetMove()
 {
-	if (GetCurrentEditingCache()->IsDynamicMeshSkeletonEnabled())
+	if (GetCurrentEditingCache() && GetCurrentEditingCache()->IsDynamicMeshSkeletonEnabled())
 	{
 		return ShouldDrawWidget();
 	}
@@ -994,7 +1016,7 @@ bool USkeletalMeshModelingToolsEditorMode::AllowWidgetMove()
 
 bool USkeletalMeshModelingToolsEditorMode::ShouldDrawWidget() const
 {
-	if (GetCurrentEditingCache()->IsDynamicMeshSkeletonEnabled())
+	if (GetCurrentEditingCache() && GetCurrentEditingCache()->IsDynamicMeshSkeletonEnabled())
 	{
 		if (GetCurrentEditingCache()->IsDynamicMeshBoneManipulationEnabled() && GetCurrentEditingCache()->GetFirstSelectedBoneIndex() != INDEX_NONE )
 		{
@@ -1011,7 +1033,7 @@ bool USkeletalMeshModelingToolsEditorMode::UsesTransformWidget() const
 {
 	// This function determines if the GetWidgetLocation() of this editor mode should be used
 	
-	if (GetCurrentEditingCache()->IsDynamicMeshSkeletonEnabled())
+	if (GetCurrentEditingCache() && GetCurrentEditingCache()->IsDynamicMeshSkeletonEnabled())
 	{
 		return true;
 	}
@@ -1021,7 +1043,7 @@ bool USkeletalMeshModelingToolsEditorMode::UsesTransformWidget() const
 
 FVector USkeletalMeshModelingToolsEditorMode::GetWidgetLocation() const
 {
-	if (GetCurrentEditingCache()->IsDynamicMeshSkeletonEnabled())
+	if (GetCurrentEditingCache() && GetCurrentEditingCache()->IsDynamicMeshSkeletonEnabled())
 	{
 		int32 BoneIndex = GetCurrentEditingCache()->GetFirstSelectedBoneIndex();
 		URefSkeletonPoser* Poser = GetCurrentEditingCache()->GetSkeletonPoser();
@@ -1045,7 +1067,7 @@ FVector USkeletalMeshModelingToolsEditorMode::GetWidgetLocation() const
 
 bool USkeletalMeshModelingToolsEditorMode::GetCustomDrawingCoordinateSystem(FMatrix& InMatrix, void* InData)
 {
-	if (GetCurrentEditingCache()->IsDynamicMeshSkeletonEnabled())
+	if (GetCurrentEditingCache() && GetCurrentEditingCache()->IsDynamicMeshSkeletonEnabled())
 	{
 		const bool bIsParentMode = Owner ? Owner->GetCoordSystem() == COORD_Parent : false;
 	
@@ -1078,7 +1100,7 @@ bool USkeletalMeshModelingToolsEditorMode::GetCustomDrawingCoordinateSystem(FMat
 
 bool USkeletalMeshModelingToolsEditorMode::GetCustomInputCoordinateSystem(FMatrix& InMatrix, void* InData)
 {
-	if (GetCurrentEditingCache()->IsDynamicMeshSkeletonEnabled())
+	if (GetCurrentEditingCache() && GetCurrentEditingCache()->IsDynamicMeshSkeletonEnabled())
 	{
 		return GetCustomDrawingCoordinateSystem(InMatrix, InData);
 	}
@@ -1088,6 +1110,11 @@ bool USkeletalMeshModelingToolsEditorMode::GetCustomInputCoordinateSystem(FMatri
 
 USkeletalMeshBackedDynamicMeshComponent* USkeletalMeshModelingToolsEditorMode::GetComponent(UObject* SourceObject)
 {
+	if (!GetCurrentEditingCache())
+	{
+		return nullptr;
+	}
+	
 	return GetCurrentEditingCache()->GetEditingMeshComponent();
 }
 
@@ -1197,17 +1224,25 @@ void USkeletalMeshModelingToolsEditorMode::SetEditingLOD(EMeshLODIdentifier InEd
 		GetCurrentEditingCache()->ApplyChanges();
 	}
 	RecreateEditingCache(InEditingLOD);
-	SkeletonReader->ExternalUpdate(GetCurrentEditingCache()->GetEditingMeshComponent()->GetRefSkeleton());
 
-	TObjectPtr<UToolTargetManager> TargetManager = GetInteractiveToolsContext(EToolsContextScope::EdMode)->TargetManager;
-	if (USkeletalMeshComponentToolTargetFactory* SkeletalMeshTargetFactory = TargetManager->FindFirstFactoryByType<USkeletalMeshComponentToolTargetFactory>())
+	if (GetCurrentEditingCache())
 	{
-		SkeletalMeshTargetFactory->SetActiveEditingLOD(InEditingLOD);
+		SkeletonReader->ExternalUpdate(GetCurrentEditingCache()->GetEditingMeshComponent()->GetRefSkeleton());
+
+		TObjectPtr<UToolTargetManager> TargetManager = GetInteractiveToolsContext(EToolsContextScope::EdMode)->TargetManager;
+		if (USkeletalMeshComponentToolTargetFactory* SkeletalMeshTargetFactory = TargetManager->FindFirstFactoryByType<USkeletalMeshComponentToolTargetFactory>())
+		{
+			SkeletalMeshTargetFactory->SetActiveEditingLOD(InEditingLOD);
+		}
 	}
 }
 
 EMeshLODIdentifier USkeletalMeshModelingToolsEditorMode::GetEditingLOD()
 {
+	if (!GetCurrentEditingCache())
+	{
+		return EMeshLODIdentifier::LOD0;
+	}
 	return GetCurrentEditingCache()->GetLOD();
 }
 
@@ -1223,6 +1258,12 @@ void USkeletalMeshModelingToolsEditorMode::RecreateEditingCache(EMeshLODIdentifi
 	{
 		return;
 	}
+
+	if (!ensure(GetSkelMeshComponent()) || !ensure(GetSkelMeshComponent()->GetSkeletalMeshAsset()))
+	{
+		UE_LOG(LogSkeletalMeshModelingTools, Error, TEXT("Failed to create editing cache for Skeletal Mesh"));
+		return;
+	}
 	
 	TSharedRef<IPersonaPreviewScene> PreviewScene = Editor.Pin()->GetPersonaToolkit()->GetPreviewScene();
 	UWorld* PreviewWorld = PreviewScene->GetWorld();
@@ -1256,16 +1297,30 @@ USkeletalMeshEditingCache* USkeletalMeshModelingToolsEditorMode::GetCurrentEditi
 
 bool USkeletalMeshModelingToolsEditorMode::HasUnappliedChanges() const
 {
+	if (!GetCurrentEditingCache())
+	{
+		return false;
+	}
+	
 	return GetCurrentEditingCache()->GetEditingMeshComponent()->IsDirty();
 }
 
 void USkeletalMeshModelingToolsEditorMode::ApplyChanges()
 {
+	if (!GetCurrentEditingCache())
+	{
+		return;
+	}
 	GetCurrentEditingCache()->ApplyChanges();
 }
 
 void USkeletalMeshModelingToolsEditorMode::DiscardChanges()
 {
+	if (!GetCurrentEditingCache())
+	{
+		return;
+	}
+	
 	GetCurrentEditingCache()->DiscardChanges();
 }
 
@@ -1289,11 +1344,19 @@ void USkeletalMeshModelingToolsEditorMode::SetApplyMode(EApplyMode InApplyMode)
 
 void USkeletalMeshModelingToolsEditorMode::HideSkeletonForTool()
 {
+	if (!GetCurrentEditingCache())
+	{
+		return;
+	}	
 	GetCurrentEditingCache()->HideSkeleton();
 }
 
 void USkeletalMeshModelingToolsEditorMode::ShowSkeletonForTool()
 {
+	if (!GetCurrentEditingCache())
+	{
+		return;
+	}
 	GetCurrentEditingCache()->ShowSkeleton();
 }
 
diff --git a/Engine/Plugins/Animation/SkeletalMeshModelingTools/Source/SkeletalMeshModelingTools/Private/SkeletalMeshModelingToolsEditorMode.h b/Engine/Plugins/Animation/SkeletalMeshModelingTools/Source/SkeletalMeshModelingTools/Private/SkeletalMeshModelingToolsEditorMode.h
--- a/Engine/Plugins/Animation/SkeletalMeshModelingTools/Source/SkeletalMeshModelingTools/Private/SkeletalMeshModelingToolsEditorMode.h
+++ b/Engine/Plugins/Animation/SkeletalMeshModelingTools/Source/SkeletalMeshModelingTools/Private/SkeletalMeshModelingToolsEditorMode.h
@@ -139,7 +139,6 @@ public:
 	
 	void HideSkeletonForTool();	
 	void ShowSkeletonForTool();
-	const TArray<FTransform>& GetComponentSpaceBoneTransforms() const;
 	void ToggleBoneManipulation(bool bEnable);
 	USkeletonModifier* GetSkeletonReader();
 	void SetSelectedBones(const TArray<FName>& InSelectedBones);
diff --git a/Engine/Plugins/Animation/SkeletalMeshModelingTools/Source/SkeletalMeshModelingTools/Private/SkeletalMeshModelingToolsEditorModeToolkit.cpp b/Engine/Plugins/Animation/SkeletalMeshModelingTools/Source/SkeletalMeshModelingTools/Private/SkeletalMeshModelingToolsEditorModeToolkit.cpp
--- a/Engine/Plugins/Animation/SkeletalMeshModelingTools/Source/SkeletalMeshModelingTools/Private/SkeletalMeshModelingToolsEditorModeToolkit.cpp
+++ b/Engine/Plugins/Animation/SkeletalMeshModelingTools/Source/SkeletalMeshModelingTools/Private/SkeletalMeshModelingToolsEditorModeToolkit.cpp
@@ -255,7 +255,7 @@ void FSkeletalMeshModelingToolsEditorModeToolkit::ResetDynamicMeshBoneTransforms
 	
 	FScopedTransaction Transaction(Title);
 	
-	if (EditorMode.IsValid())
+	if (EditorMode.IsValid() && EditorMode->GetCurrentEditingCache())
 	{
 		EditorMode->GetCurrentEditingCache()->ResetDynamicMeshBoneTransforms(bSelectedOnly);
 	}
@@ -900,7 +900,7 @@ FReply FSkeletalMeshModelingToolsEditorModeToolkit::OnApplyButtonPressed()
 
 FText FSkeletalMeshModelingToolsEditorModeToolkit::GetApplyButtonText() const
 {
-	int32 ChangeCount = EditorMode->GetCurrentEditingCache()->GetEditingMeshComponent()->GetChangeCount();
+	int32 ChangeCount = EditorMode->GetCurrentEditingCache() ? EditorMode->GetCurrentEditingCache()->GetEditingMeshComponent()->GetChangeCount() : 0;
 	if (ChangeCount == 0)
 	{
 		return FText::Format(LOCTEXT("ApplyButtonZeroChangeText", "No Pending Changes"), FText::AsNumber(ChangeCount));
