[SkeletalMeshModelingTools] enable new morph target for editing immediately after it is created

#Synced-CL 48737829

diff --git a/Engine/Plugins/Animation/SkeletalMeshModelingTools/Source/SkeletalMeshModelingTools/Private/SkeletalMeshModelingToolsEditorMode.cpp b/Engine/Plugins/Animation/SkeletalMeshModelingTools/Source/SkeletalMeshModelingTools/Private/SkeletalMeshModelingToolsEditorMode.cpp
--- a/Engine/Plugins/Animation/SkeletalMeshModelingTools/Source/SkeletalMeshModelingTools/Private/SkeletalMeshModelingToolsEditorMode.cpp
+++ b/Engine/Plugins/Animation/SkeletalMeshModelingTools/Source/SkeletalMeshModelingTools/Private/SkeletalMeshModelingToolsEditorMode.cpp
@@ -263,6 +263,7 @@ USkeletalMeshModelingToolsEditorMode::~USkeletalMeshModelingToolsEditorMode()
 void USkeletalMeshModelingToolsEditorMode::Initialize()
 {
 	UBaseLegacyWidgetEdMode::Initialize();
+	SetFlags(RF_Transactional);
 }
 
 
@@ -521,7 +522,6 @@ TArray<FName> USkeletalMeshModelingToolsEditorMode::GetSelectedBones() const
 	return SelectedBones;
 }
 
-
 FName USkeletalMeshModelingToolsEditorMode::GetEditingMorphTarget()
 {
 	if (!GetCurrentEditingCache())
@@ -537,8 +537,12 @@ FName USkeletalMeshModelingToolsEditorMode::GetEditingMorphTarget()
 	return EditingMorphTarget;
 }
 
+
+
 void USkeletalMeshModelingToolsEditorMode::HandleSetEditingMorphTarget(FName InMorphTarget)
 {
+	// Such that EditingMorphTarget is recorded for undo/redo 
+	Modify();
 	FName OldEditingMorphTarget = EditingMorphTarget;
 	EditingMorphTarget = InMorphTarget;
 
@@ -569,7 +573,11 @@ float USkeletalMeshModelingToolsEditorMode::GetMorphTargetWeight(FName MorphTarg
 
 FName USkeletalMeshModelingToolsEditorMode::HandleAddMorphTarget(FName InName)
 {
-	return GetCurrentEditingCache()->AddMorphTarget(InName);
+	FName NewMorphTarget = GetCurrentEditingCache()->AddMorphTarget(InName);
+	
+	HandleSetEditingMorphTarget(NewMorphTarget);
+	
+	return NewMorphTarget;
 }
 
 FName USkeletalMeshModelingToolsEditorMode::HandleRenameMorphTarget(FName OldName, FName NewName)
@@ -1290,6 +1298,8 @@ void USkeletalMeshModelingToolsEditorMode::RecreateEditingCache(EMeshLODIdentifi
 	CurrentEditingCache->GetNotifier()->HandleNotification(GetSelectedBones(), ESkeletalMeshNotifyType::BonesSelected);
 }
 
+
+
 USkeletalMeshEditingCache* USkeletalMeshModelingToolsEditorMode::GetCurrentEditingCache() const
 {
 	return CurrentEditingCache;
diff --git a/Engine/Plugins/Animation/SkeletalMeshModelingTools/Source/SkeletalMeshModelingTools/Private/SkeletalMeshModelingToolsEditorMode.h b/Engine/Plugins/Animation/SkeletalMeshModelingTools/Source/SkeletalMeshModelingTools/Private/SkeletalMeshModelingToolsEditorMode.h
--- a/Engine/Plugins/Animation/SkeletalMeshModelingTools/Source/SkeletalMeshModelingTools/Private/SkeletalMeshModelingToolsEditorMode.h
+++ b/Engine/Plugins/Animation/SkeletalMeshModelingTools/Source/SkeletalMeshModelingTools/Private/SkeletalMeshModelingToolsEditorMode.h
@@ -9,6 +9,7 @@
 #include "Interfaces/ISKMBackedDynaMeshComponentProvider.h"
 #include "Animation/DebugSkelMeshComponent.h"
 #include "Changes/ValueWatcher.h"
+#include "Misc/Change.h"
 #include "TargetInterfaces/MeshTargetInterfaceTypes.h"
 
 #include "SkeletalMeshModelingToolsEditorMode.generated.h"
@@ -181,7 +182,7 @@ private:
 	TWeakPtr<ISkeletalMeshEditorBinding> EditorBinding;
 	TWeakObjectPtr<USkeletalMesh> SkeletalMesh;
 
-	UPROPERTY()
+	UPROPERTY(NonTransactional)
 	TObjectPtr<USkeletalMeshEditingCache> CurrentEditingCache = nullptr;
 
 	void RecreateEditingCache(EMeshLODIdentifier InLOD);
@@ -191,7 +192,8 @@ private:
 	// Represents the collection of bones currently selected in the viewport for whichever skeleton that is currently active
 	// Helps maintain consistent bone selection state as different skeleton tree widget activates
 	TArray<FName> SelectedBones;
-	
+
+	UPROPERTY()
 	FName EditingMorphTarget = NAME_None;
 	
 	// Stylus support is currently disabled; this is left in for reference if/when it is brought back
@@ -207,7 +209,7 @@ private:
 	bool IsSkeletalMeshBoneManipulationEnabled();
 
 	// A dummy read-only modifier to host the current edited skeleton for RefSkeletonTree UI
-	UPROPERTY()
+	UPROPERTY(NonTransactional)
 	TObjectPtr<USkeletonModifier> SkeletonReader = nullptr;
 
 	TValueWatcher<int32> CacheChangeCountWatcher;
