Morph Target Manager: Ensure that search is case-insensitive.

#rnx
#Synced-CL 48772759

diff --git a/Engine/Plugins/Animation/SkeletalMeshModelingTools/Source/SkeletalMeshModelingTools/Private/SMorphTargetManager.cpp b/Engine/Plugins/Animation/SkeletalMeshModelingTools/Source/SkeletalMeshModelingTools/Private/SMorphTargetManager.cpp
--- a/Engine/Plugins/Animation/SkeletalMeshModelingTools/Source/SkeletalMeshModelingTools/Private/SMorphTargetManager.cpp
+++ b/Engine/Plugins/Animation/SkeletalMeshModelingTools/Source/SkeletalMeshModelingTools/Private/SMorphTargetManager.cpp
@@ -370,7 +370,7 @@ void SMorphTargetManager::RefreshList()
 		bool bAddToList = true;
 		if (!FilterText.IsEmpty())
 		{
-			FRegexPattern Pattern(FilterText.ToString());
+			FRegexPattern Pattern(FilterText.ToString(), ERegexPatternFlags::CaseInsensitive);
 			FRegexMatcher Matcher(Pattern, MorphTarget.ToString());
 
 			bAddToList = Matcher.FindNext();
@@ -442,7 +442,7 @@ FText SMorphTargetManager::GetHighlightText(FText InName) const
 	FText FilterText = GetFilterText();
 	if (!FilterText.IsEmpty())
 	{
-		FRegexPattern Pattern(FilterText.ToString());
+		FRegexPattern Pattern(FilterText.ToString(), ERegexPatternFlags::CaseInsensitive);
 		FString Name = InName.ToString();
 		FRegexMatcher Matcher(Pattern, Name);
 		if (Matcher.FindNext())
