[SkeletalMeshModelingTools] added a nullcheck

diff --git a/Engine/Plugins/Animation/SkeletalMeshModelingTools/Source/SkeletalMeshModelingTools/Private/SkeletalMeshModelingToolsEditorMode.cpp b/Engine/Plugins/Animation/SkeletalMeshModelingTools/Source/SkeletalMeshModelingTools/Private/SkeletalMeshModelingToolsEditorMode.cpp
--- a/Engine/Plugins/Animation/SkeletalMeshModelingTools/Source/SkeletalMeshModelingTools/Private/SkeletalMeshModelingToolsEditorMode.cpp
+++ b/Engine/Plugins/Animation/SkeletalMeshModelingTools/Source/SkeletalMeshModelingTools/Private/SkeletalMeshModelingToolsEditorMode.cpp
@@ -815,52 +815,55 @@ bool USkeletalMeshModelingToolsEditorMode::ComputeBoundingBoxForViewportFocus(AA
 		}
 	}
 
-	TArray<FName> Selection = GetCurrentEditingCache()->GetSelectedBones();
-	
-	
-	// focus using selected bones in skel mesh editor
-	if (!Selection.IsEmpty())
+	if (GetCurrentEditingCache())
 	{
-		const FReferenceSkeleton& RefSkeleton = GetCurrentEditingCache()->GetEditingMeshComponent()->GetRefSkeleton();
-		TArray<FTransform> WorldTransforms = GetCurrentEditingCache()->GetComponentSpaceBoneTransforms();
-		FTransform MeshTransform = GetCurrentEditingCache()->GetTransform();
-		for (FTransform& Transform : WorldTransforms)
+		TArray<FName> Selection = GetCurrentEditingCache()->GetSelectedBones();
+	
+		// focus using selected bones in skel mesh editor
+		if (!Selection.IsEmpty())
 		{
-			Transform *= MeshTransform;
-		}
+			const FReferenceSkeleton& RefSkeleton = GetCurrentEditingCache()->GetEditingMeshComponent()->GetRefSkeleton();
+			TArray<FTransform> WorldTransforms = GetCurrentEditingCache()->GetComponentSpaceBoneTransforms();
+			FTransform MeshTransform = GetCurrentEditingCache()->GetTransform();
+			for (FTransform& Transform : WorldTransforms)
+			{
+				Transform *= MeshTransform;
+			}
 		
-		TArray<FName> AllChildren;
+			TArray<FName> AllChildren;
 
-		for (const FName& BoneName: Selection)
-		{
-			const int32 BoneIndex = RefSkeleton.FindBoneIndex(BoneName);
-			if (BoneIndex > INDEX_NONE && WorldTransforms.IsValidIndex(BoneIndex))
+			for (const FName& BoneName: Selection)
 			{
-				// enlarge box
-				InOutBox += WorldTransforms[BoneIndex].GetLocation();
-
-				// get direct children
-				TArray<int32> Children;
-				RefSkeleton.GetDirectChildBones(BoneIndex, Children);
-				Algo::Transform(Children, AllChildren, [&RefSkeleton](int ChildrenIndex)
+				const int32 BoneIndex = RefSkeleton.FindBoneIndex(BoneName);
+				if (BoneIndex > INDEX_NONE && WorldTransforms.IsValidIndex(BoneIndex))
 				{
-					return RefSkeleton.GetBoneName(ChildrenIndex);
-				});
+					// enlarge box
+					InOutBox += WorldTransforms[BoneIndex].GetLocation();
+
+					// get direct children
+					TArray<int32> Children;
+					RefSkeleton.GetDirectChildBones(BoneIndex, Children);
+					Algo::Transform(Children, AllChildren, [&RefSkeleton](int ChildrenIndex)
+					{
+						return RefSkeleton.GetBoneName(ChildrenIndex);
+					});
+				}
 			}
-		}
 
-		// enlarge box using direct children
-		for (const FName& BoneName: AllChildren)
-		{
-			const int32 BoneIndex = RefSkeleton.FindBoneIndex(BoneName);
-			if (BoneIndex > INDEX_NONE && WorldTransforms.IsValidIndex(BoneIndex))
+			// enlarge box using direct children
+			for (const FName& BoneName: AllChildren)
 			{
-				InOutBox += WorldTransforms[BoneIndex].GetLocation();	
+				const int32 BoneIndex = RefSkeleton.FindBoneIndex(BoneName);
+				if (BoneIndex > INDEX_NONE && WorldTransforms.IsValidIndex(BoneIndex))
+				{
+					InOutBox += WorldTransforms[BoneIndex].GetLocation();	
+				}
 			}
-		}
 		
-		return true; 
+			return true; 
+		}
 	}
+
 	
 	return Super::ComputeBoundingBoxForViewportFocus(Actor, PrimitiveComponent, InOutBox);
 }
