[SkeletalMeshModelingTools] Integrated geometry seleciton manager, supports editing base mesh or actively editing morph target using element selection modes

#Synced-CL 48737829

diff --git a/Engine/Plugins/Animation/SkeletalMeshModelingTools/Source/SkeletalMeshModelingTools/Private/Components/SKMBackedDynaMeshComponent.cpp b/Engine/Plugins/Animation/SkeletalMeshModelingTools/Source/SkeletalMeshModelingTools/Private/Components/SKMBackedDynaMeshComponent.cpp
--- a/Engine/Plugins/Animation/SkeletalMeshModelingTools/Source/SkeletalMeshModelingTools/Private/Components/SKMBackedDynaMeshComponent.cpp
+++ b/Engine/Plugins/Animation/SkeletalMeshModelingTools/Source/SkeletalMeshModelingTools/Private/Components/SKMBackedDynaMeshComponent.cpp
@@ -13,6 +13,7 @@
 #include "Misc/ITransaction.h"
 #include "SkeletalMeshEditorSubsystem.h"
 #include "ToolTargets/SkeletalMeshToolTarget.h"
+#include "SkeletalMesh/SkeletalMeshToolsHelper.h"
 
 #include UE_INLINE_GENERATED_CPP_BY_NAME(SKMBackedDynaMeshComponent)
 
@@ -361,6 +362,81 @@ void USkeletalMeshBackedDynamicMeshComponent::ForwardVisibilityChangeRequest(boo
 	OnRequestingVisibilityChangeDelegate.Broadcast(bInVisible);
 }
 
+void USkeletalMeshBackedDynamicMeshComponent::HandleGeometryUpdate(
+	const FName EditingMorphTargetName,
+	const FDynamicMesh3& PosedMesh,
+	const TArray<FMatrix>& BoneMatrices,
+	const TMap<FName, float>& MorphTargetWeights)
+{
+	using namespace UE::Geometry;
+	using namespace SkeletalMeshToolsHelper;
+
+	FDynamicMeshMorphTargetAttribute* MorphTargetAttribute = nullptr;
+	float EditingMorphTargetWeight = 1.0f;
+	
+	ProcessMesh([&](const FDynamicMesh3& Mesh)
+		{
+			MorphTargetAttribute = Mesh.Attributes()->GetMorphTargetAttribute(EditingMorphTargetName);
+			if (MorphTargetAttribute)
+			{
+				const float* FoundWeight = MorphTargetWeights.Find(EditingMorphTargetName);
+				if (ensure(FoundWeight))
+				{
+					EditingMorphTargetWeight = *FoundWeight;
+				}
+			}
+		});
+
+	if (FMath::IsNearlyZero(EditingMorphTargetWeight))
+	{
+		return;
+	}
+	
+	FChangeScope ChangeScope(this);
+	
+	// Mesh change
+	EditMesh([&](FDynamicMesh3& Mesh)
+		{
+			TFunction<void(FVertInfo, const FVector&)> WriteFunction;
+			TMap<FName, float> UnposeMorphTargetWeights = MorphTargetWeights;
+			
+			if (MorphTargetAttribute)
+			{
+				WriteFunction = [&](FVertInfo VertInfo, const FVector& UnposedVertPos)
+				{
+					int32 VertID = VertInfo.VertID;
+
+					FVector NewDelta = UnposedVertPos - FVector(Mesh.GetVertex(VertID));
+					NewDelta /= EditingMorphTargetWeight;
+
+					MorphTargetAttribute->SetValue(VertID, NewDelta);
+				};
+
+				// Exclude the morph target we are trying to extract
+				UnposeMorphTargetWeights.Remove(EditingMorphTargetName);
+			}
+			else
+			{
+				WriteFunction = [&](FVertInfo VertInfo, const FVector& UnposedVertPos)
+				{
+					int32 VertID = VertInfo.VertID;
+
+					Mesh.SetVertex(VertID, UnposedVertPos);
+				};
+			}		
+
+			GetUnposedMesh(WriteFunction, PosedMesh, Mesh, BoneMatrices, NAME_None, UnposeMorphTargetWeights);
+		});
+			
+	// Tracker change
+	if (MorphTargetAttribute)
+	{
+		Trackers.MorphTargetChangeTracker.HandleEditMorphTarget(EditingMorphTargetName);
+	}
+	
+	MarkDirtyInternal();
+}
+
 FName USkeletalMeshBackedDynamicMeshComponent::AddMorphTarget(FName InName)
 {
 	FName ActualName = GetAvailableMorphTargetName(InName);
diff --git a/Engine/Plugins/Animation/SkeletalMeshModelingTools/Source/SkeletalMeshModelingTools/Private/Components/SKMBackedDynaMeshComponent.h b/Engine/Plugins/Animation/SkeletalMeshModelingTools/Source/SkeletalMeshModelingTools/Private/Components/SKMBackedDynaMeshComponent.h
--- a/Engine/Plugins/Animation/SkeletalMeshModelingTools/Source/SkeletalMeshModelingTools/Private/Components/SKMBackedDynaMeshComponent.h
+++ b/Engine/Plugins/Animation/SkeletalMeshModelingTools/Source/SkeletalMeshModelingTools/Private/Components/SKMBackedDynaMeshComponent.h
@@ -75,6 +75,13 @@ public:
 	void HandleSkeletonChange(USkeletonModifier* InModifier);
 	void ForwardVisibilityChangeRequest(bool bInVisible);
 
+	void HandleGeometryUpdate(
+		const FName EditingMorphTargetName,
+		const FDynamicMesh3& PosedMesh,
+		const TArray<FMatrix>& BoneMatrices,
+		const TMap<FName, float>& MorphTargetWeights
+		);
+
 	FName AddMorphTarget(FName InName);
 	FName RenameMorphTarget(FName InOldName, FName InNewName);
 	void RemoveMorphTargets(const TArray<FName>& InNames);
diff --git a/Engine/Plugins/Animation/SkeletalMeshModelingTools/Source/SkeletalMeshModelingTools/Private/Selection/SKMPreviewDynaMeshSelector.cpp b/Engine/Plugins/Animation/SkeletalMeshModelingTools/Source/SkeletalMeshModelingTools/Private/Selection/SKMPreviewDynaMeshSelector.cpp
new file mode 100644
--- /dev/null
+++ b/Engine/Plugins/Animation/SkeletalMeshModelingTools/Source/SkeletalMeshModelingTools/Private/Selection/SKMPreviewDynaMeshSelector.cpp
@@ -0,0 +1,96 @@
+// Copyright Epic Games, Inc. All Rights Reserved.
+
+
+#include "Selection/SKMPreviewDynaMeshSelector.h"
+
+#include "SkeletalMeshEditingCache.h"
+#include "SkeletalMeshModelingToolsEditorMode.h"
+#include "Components/SKMBackedDynaMeshComponent.h"
+#include "Selection/DynamicMeshPolygroupTransformer.h"
+
+#define LOCTEXT_NAMESPACE "FSkeletalMeshPreviewDynamicMeshSelector"
+
+using namespace UE::Geometry;
+
+IGeometrySelectionTransformer* FSkeletalMeshPreviewDynamicMeshSelector::InitializeTransformation(const FGeometrySelection& Selection)
+{
+	check(!ActiveTransformer);
+
+	if (Selection.TopologyType == EGeometryTopologyType::Polygroup)
+	{
+		ActiveTransformer = MakeShared<FDynamicMeshPolygroupTransformer>();
+	}
+	else
+	{
+		ActiveTransformer = MakeShared<FBasicDynamicMeshSelectionTransformer>();
+	}
+	
+	ActiveTransformer->bEnableSelectionTransformDrawing = true;
+
+	ActiveTransformer->Initialize(this);
+	ActiveTransformer->OnEndTransformFunc = [this](IToolsContextTransactionsAPI* TransactionsAPI) 
+		{
+			CommitMeshTransform();
+		};
+	return ActiveTransformer.Get();
+}
+
+void FSkeletalMeshPreviewDynamicMeshSelector::ShutdownTransformation(IGeometrySelectionTransformer* Transformer)
+{
+	ActiveTransformer.Reset();
+}
+
+
+void FSkeletalMeshPreviewDynamicMeshSelector::CommitMeshTransform()
+{
+	USkeletalMeshEditingCache* Cache = WeakEditorMode->GetCurrentEditingCache();
+
+	if (!Cache)
+	{
+		return;
+	}
+
+	Cache->HandleGeometryUpdate(TargetMesh->GetMeshRef(), WeakEditorMode->GetEditingMorphTarget());
+}
+
+void FSkeletalMeshPreviewDynamicMeshSelectorFactory::Init(USkeletalMeshModelingToolsEditorMode* EditorMode)
+{
+	WeakEditorMode = EditorMode;
+}
+
+bool FSkeletalMeshPreviewDynamicMeshSelectorFactory::CanBuildForTarget(FGeometryIdentifier TargetIdentifier) const
+{
+	if (TargetIdentifier.TargetType == FGeometryIdentifier::ETargetType::PrimitiveComponent)
+	{
+		UDynamicMeshComponent* Component = TargetIdentifier.GetAsComponentType<UDynamicMeshComponent>();
+		if (Component)
+		{
+			return true;
+		}
+	}
+	return false;
+}
+
+TUniquePtr<IGeometrySelector> FSkeletalMeshPreviewDynamicMeshSelectorFactory::BuildForTarget(FGeometryIdentifier TargetIdentifier) const
+{
+	UDynamicMeshComponent* Component = TargetIdentifier.GetAsComponentType<UDynamicMeshComponent>();
+	check(Component != nullptr);
+
+	TUniquePtr<FSkeletalMeshPreviewDynamicMeshSelector> Selector = MakeUnique<FSkeletalMeshPreviewDynamicMeshSelector>();
+	Selector->Initialize(TargetIdentifier, Component->GetDynamicMesh(),
+		[TargetIdentifier]() -> FTransformSRT3d
+		{ 
+			UDynamicMeshComponent* Component = TargetIdentifier.GetAsComponentType<UDynamicMeshComponent>();
+			return IsValid(Component) ? (FTransformSRT3d)Component->GetComponentTransform() : FTransformSRT3d::Identity(); 
+		}
+	);
+
+	Selector->WeakEditorMode = WeakEditorMode;
+	return Selector;
+}
+
+
+
+
+
+#undef LOCTEXT_NAMESPACE 
diff --git a/Engine/Plugins/Animation/SkeletalMeshModelingTools/Source/SkeletalMeshModelingTools/Private/Selection/SKMPreviewDynaMeshSelector.h b/Engine/Plugins/Animation/SkeletalMeshModelingTools/Source/SkeletalMeshModelingTools/Private/Selection/SKMPreviewDynaMeshSelector.h
new file mode 100644
--- /dev/null
+++ b/Engine/Plugins/Animation/SkeletalMeshModelingTools/Source/SkeletalMeshModelingTools/Private/Selection/SKMPreviewDynaMeshSelector.h
@@ -0,0 +1,43 @@
+// Copyright Epic Games, Inc. All Rights Reserved.
+
+#pragma once
+
+#include "CoreMinimal.h"
+#include "ModelingComponents/Public/Selection/DynamicMeshSelector.h"
+
+#define UE_API SKELETALMESHMODELINGTOOLS_API
+
+
+class USkeletalMeshModelingToolsEditorMode;
+
+class FSkeletalMeshPreviewDynamicMeshSelector : public FBaseDynamicMeshSelector
+{
+public:
+	
+	UE_API virtual IGeometrySelectionTransformer* InitializeTransformation(const FGeometrySelection& Selection) override;
+	UE_API virtual void ShutdownTransformation(IGeometrySelectionTransformer* Transformer) override;
+	
+	TWeakObjectPtr<USkeletalMeshModelingToolsEditorMode> WeakEditorMode;
+	
+protected:	
+	TSharedPtr<FBasicDynamicMeshSelectionTransformer> ActiveTransformer;
+	UE_API void CommitMeshTransform();
+	
+};
+
+
+
+class FSkeletalMeshPreviewDynamicMeshSelectorFactory : public IGeometrySelectorFactory
+{
+public:
+	void Init(USkeletalMeshModelingToolsEditorMode* EditorMode);
+	
+	virtual bool CanBuildForTarget(FGeometryIdentifier TargetIdentifier) const;
+
+	virtual TUniquePtr<IGeometrySelector> BuildForTarget(FGeometryIdentifier TargetIdentifier) const;
+protected:
+	TWeakObjectPtr<USkeletalMeshModelingToolsEditorMode> WeakEditorMode;
+};
+
+
+#undef UE_API
diff --git a/Engine/Plugins/Animation/SkeletalMeshModelingTools/Source/SkeletalMeshModelingTools/Private/SkeletalMeshEditingCache.cpp b/Engine/Plugins/Animation/SkeletalMeshModelingTools/Source/SkeletalMeshModelingTools/Private/SkeletalMeshEditingCache.cpp
--- a/Engine/Plugins/Animation/SkeletalMeshModelingTools/Source/SkeletalMeshModelingTools/Private/SkeletalMeshEditingCache.cpp
+++ b/Engine/Plugins/Animation/SkeletalMeshModelingTools/Source/SkeletalMeshModelingTools/Private/SkeletalMeshEditingCache.cpp
@@ -97,8 +97,9 @@ void USkeletalMeshEditingCache::Spawn(UWorld* World, USkeletalMeshComponent* InS
 	
 	EditingMeshComponent->SetVisibility(false);
 	PreviewMesh->SetVisible(false);
-	bEnableDynamicMesh = false;
-	bEnableDynamicMeshSkeleton = false;
+	bIsDynamicMeshEnabled = false;
+	bForceEnableDynamicMesh = false;
+	bIsDynamicMeshSkeletonEnabled = false;
 
 	bCacheVisibility = true;
 	CacheSkeletonDrawMode = ESkeletonDrawMode::Default;
@@ -121,7 +122,7 @@ void USkeletalMeshEditingCache::Spawn(UWorld* World, USkeletalMeshComponent* InS
 			},
 		[this](int32 ChangeCount)
 			{
-				bEnableDynamicMesh = ChangeCount > 0;
+				
 				UDynamicMesh* DynamicMesh = EditingMeshComponent->GetDynamicMesh();
 				FDynamicMesh3 PreviewMeshObject;
 				DynamicMesh->ProcessMesh([&](const FDynamicMesh3& ReadMesh)
@@ -130,10 +131,24 @@ void USkeletalMeshEditingCache::Spawn(UWorld* World, USkeletalMeshComponent* InS
 					});
 
 				PreviewMesh->ReplaceMesh(MoveTemp(PreviewMeshObject));	
-				
+				DeformPreviewMesh(GetComponentSpaceBoneTransforms(), GetMorphTargetWeights());
 			},
 		GetEditingMeshComponent()->GetChangeCount());
 
+	DynamicMeshStateWatcher.Initialize(
+		[this]()
+			{
+				FDynamicMeshState State;
+				State.ChangeCount = GetEditingMeshComponent()->GetChangeCount();
+				State.bForceEnable = bForceEnableDynamicMesh;
+				return State;
+			},
+		[this](FDynamicMeshState State)
+			{
+				bIsDynamicMeshEnabled = State.bForceEnable || State.ChangeCount > 0;
+			},
+		FDynamicMeshState(GetEditingMeshComponent()->GetChangeCount(), bForceEnableDynamicMesh));
+	
 	SkeletonChangeCountWatcher.Initialize(
 		[this]()
 			{
@@ -141,10 +156,10 @@ void USkeletalMeshEditingCache::Spawn(UWorld* World, USkeletalMeshComponent* InS
 			},
 		[this](int32 SkeletonChangeCount)
 			{
-				bEnableDynamicMeshSkeleton = SkeletonChangeCount > 0;
+				bIsDynamicMeshSkeletonEnabled = SkeletonChangeCount > 0;
 				RefSkeletonPoser->SetRefSkeleton(GetEditingMeshComponent()->GetRefSkeleton());
 				UpdateSelectedBoneIndices();
-				if (bEnableDynamicMeshSkeleton)
+				if (bIsDynamicMeshSkeletonEnabled)
 				{
 					SkeletalMeshSkeletonNotifierBindScope.Reset();
 				}
@@ -169,35 +184,34 @@ void USkeletalMeshEditingCache::Spawn(UWorld* World, USkeletalMeshComponent* InS
 		[this]()
 			{
 				FMeshVisibilityState State;
-				State.bEnableDynamicMesh = bEnableDynamicMesh;
+				State.bIsDynamicMeshEnabled = bIsDynamicMeshEnabled;
 				State.bCacheVisibility = bCacheVisibility;
 				return State;
 			},
 		[this](FMeshVisibilityState State)
 			{
 				PreviewMesh->SetVisible(
-					State.bEnableDynamicMesh && State.bCacheVisibility);
+					State.bIsDynamicMeshEnabled && State.bCacheVisibility);
 
 				if (SkeletalMeshComponent.IsValid())
 				{
 					SkeletalMeshComponent->SetVisibility(
-						!State.bEnableDynamicMesh && State.bCacheVisibility);
+						!State.bIsDynamicMeshEnabled && State.bCacheVisibility);
 				}
 			},
-		FMeshVisibilityState(bEnableDynamicMesh, bCacheVisibility)
-	);
+		FMeshVisibilityState(bIsDynamicMeshEnabled, bCacheVisibility));
 
 	SkeletonVisibilityWatcher.Initialize(
 		[this]()
 			{
 				FSkeletonVisibilityState State;
-				State.bEnableDynamicMeshSkeleton = bEnableDynamicMeshSkeleton;
+				State.bIsDynamicMeshSkeletonEnabled = bIsDynamicMeshSkeletonEnabled;
 				State.CacheSkeletonDrawMode = CacheSkeletonDrawMode;
 				return State;
 			},
 		[this](FSkeletonVisibilityState State)
 			{
-				if (State.bEnableDynamicMeshSkeleton)
+				if (State.bIsDynamicMeshSkeletonEnabled)
 				{
 					DynamicMeshSkeletonDrawMode = State.CacheSkeletonDrawMode;
 					SetSkeletalMeshSkeletonDrawMode(ESkeletonDrawMode::Hidden);
@@ -208,23 +222,23 @@ void USkeletalMeshEditingCache::Spawn(UWorld* World, USkeletalMeshComponent* InS
 					SetSkeletalMeshSkeletonDrawMode(State.CacheSkeletonDrawMode);
 				}
 			},
-		FSkeletonVisibilityState(bEnableDynamicMeshSkeleton, CacheSkeletonDrawMode)
+		FSkeletonVisibilityState(bIsDynamicMeshSkeletonEnabled, CacheSkeletonDrawMode)
 	);
 
 	BoneManipulationWatcher.Initialize(
 		[this]()
 			{
 				FBoneManipulationState State;
-				State.bEnableDynamicMeshSkeleton = bEnableDynamicMeshSkeleton;
+				State.bIsDynamicMeshSkeletonEnabled = bIsDynamicMeshSkeletonEnabled;
 				State.bCacheBoneManipulation = bCacheBoneManipulation;
 				return State;
 			},
 		[this](FBoneManipulationState State)
 			{
 				Delegates.ToggleSkeletalMeshBoneManipulationDelegate.ExecuteIfBound(
-					!State.bEnableDynamicMeshSkeleton && bCacheBoneManipulation);
+					!State.bIsDynamicMeshSkeletonEnabled && bCacheBoneManipulation);
 			},
-		FBoneManipulationState(bEnableDynamicMeshSkeleton, bCacheBoneManipulation)
+		FBoneManipulationState(bIsDynamicMeshSkeletonEnabled, bCacheBoneManipulation)
 	);
 
 	PreviewMeshVisibilityWatcher.Initialize(
@@ -297,13 +311,14 @@ void USkeletalMeshEditingCache::DiscardChanges()
 void USkeletalMeshEditingCache::Tick()
 {
 	ChangeCountWatcher.CheckAndUpdate();
+	DynamicMeshStateWatcher.CheckAndUpdate();
 	SkeletonChangeCountWatcher.CheckAndUpdate();
 	MeshVisibilityWatcher.CheckAndUpdate();
 	SkeletonVisibilityWatcher.CheckAndUpdate();
 	BoneManipulationWatcher.CheckAndUpdate();
 	PreviewMeshVisibilityWatcher.CheckAndUpdate();
 	
-	if (bEnableDynamicMesh)
+	if (bIsDynamicMeshEnabled)
 	{
 		PoseChangeDetector.CheckPose( GetComponentSpaceBoneTransforms(),  GetMorphTargetWeights());
 	}
@@ -311,7 +326,7 @@ void USkeletalMeshEditingCache::Tick()
 
 bool USkeletalMeshEditingCache::HandleClick(HHitProxy* HitProxy)
 {
-	if (bEnableDynamicMeshSkeleton)
+	if (bIsDynamicMeshSkeletonEnabled)
 	{
 		TArray<FName> SelectedBones;	
 		if (const HBoneHitProxy* BoneHitProxy = HitProxyCast<HBoneHitProxy>(HitProxy))
@@ -396,7 +411,7 @@ void USkeletalMeshEditingCache::Render(FPrimitiveDrawInterface* PDI, TFunction<v
 
 void USkeletalMeshEditingCache::DrawHUD(FEditorViewportClient* ViewportClient, FViewport* Viewport, const FSceneView* View, FCanvas* Canvas)
 {
-	if (bEnableDynamicMeshSkeleton)
+	if (bIsDynamicMeshSkeletonEnabled)
 	{
 		// See FSkeletonSelectionEditMode::DrawHUD
 		const FReferenceSkeleton& RefSkeleton = GetEditingMeshComponent()->GetRefSkeleton();
@@ -447,12 +462,21 @@ void USkeletalMeshEditingCache::HandleNotification(const TArray<FName>& BoneName
 	}
 }
 
+void USkeletalMeshEditingCache::ToggleForceEnableDynamicMesh(bool bEnable)
+{
+	bForceEnableDynamicMesh = bEnable;
+}
 
 USkeletalMeshBackedDynamicMeshComponent* USkeletalMeshEditingCache::GetEditingMeshComponent() const
 {
 	return EditingMeshComponent;
 }
 
+UPrimitiveComponent* USkeletalMeshEditingCache::GetPreviewMeshComponent() const
+{
+	return PreviewMesh->GetRootComponent();
+}
+
 
 void USkeletalMeshEditingCache::HandleVisibilityChangeRequest(bool bVisible)
 {
@@ -474,7 +498,7 @@ void USkeletalMeshEditingCache::HandlePoseChangeDetectorEvent(SkeletalMeshToolsH
 
 const TArray<FTransform>& USkeletalMeshEditingCache::GetComponentSpaceBoneTransforms() const
 {
-	if (bEnableDynamicMeshSkeleton)
+	if (bIsDynamicMeshSkeletonEnabled)
 	{
 		return RefSkeletonPoser->GetComponentSpaceTransforms();
 	}
@@ -506,7 +530,7 @@ void USkeletalMeshEditingCache::ToggleBoneManipulation(bool bEnable)
 
 bool USkeletalMeshEditingCache::IsDynamicMeshSkeletonEnabled() const
 {
-	return bEnableDynamicMeshSkeleton;
+	return bIsDynamicMeshSkeletonEnabled;
 }
 
 bool USkeletalMeshEditingCache::IsDynamicMeshBoneManipulationEnabled() const
@@ -526,7 +550,7 @@ TArray<FName> USkeletalMeshEditingCache::GetSelectedBones() const
 
 void USkeletalMeshEditingCache::ResetDynamicMeshBoneTransforms(bool bSelectedOnly)
 {
-	if (!bEnableDynamicMeshSkeleton)
+	if (!bIsDynamicMeshSkeletonEnabled)
 	{
 		return;
 	}
@@ -728,6 +752,23 @@ TArray<FName> USkeletalMeshEditingCache::DuplicateMorphTargets(const TArray<FNam
 	return EditingMeshComponent->DuplicateMorphTargets(InNames);
 }
 
+void USkeletalMeshEditingCache::HandleGeometryUpdate(const FDynamicMesh3& InMesh, FName MorphTargetName)
+{
+	using namespace SkeletalMeshToolsHelper;
+	
+	const TArray<FTransform>& ComponentSpaceTransforms = GetComponentSpaceBoneTransforms();
+	const TArray<FTransform>& ComponentSpaceTransformsRefPose = GetComponentSpaceBoneTransformsRefPose();
+
+	TArray<FMatrix> BoneMatrices = ComputeBoneMatrices(
+			ComponentSpaceTransformsRefPose,
+			ComponentSpaceTransforms
+			);
+
+	TMap<FName, float> MorphTargetWeights = GetMorphTargetWeights();
+
+	EditingMeshComponent->HandleGeometryUpdate(MorphTargetName, InMesh, BoneMatrices, MorphTargetWeights);
+}
+
 URefSkeletonPoser* USkeletalMeshEditingCache::GetSkeletonPoser() const
 {
 	return RefSkeletonPoser;
@@ -781,8 +822,14 @@ void USkeletalMeshEditingCache::DeformPreviewMesh(const TArray<FTransform>& Comp
 		};
 		
 	constexpr bool bRebuildSpatial = false;
-	PreviewMesh->DeferredEditMesh(DeformPreviewMeshFunc, bRebuildSpatial);
-	PreviewMesh->NotifyDeferredEditCompleted(UPreviewMesh::ERenderUpdateMode::FastUpdate, EMeshRenderAttributeFlags::Positions | EMeshRenderAttributeFlags::VertexNormals, bRebuildSpatial);	
+	
+	// Using UDynamicMeshComponent::EditMesh() to make sure dependent systems, such as seleciton manager, receive a notification
+	UDynamicMeshComponent* DynamicMeshComponent = CastChecked<UDynamicMeshComponent>(PreviewMesh->GetRootComponent());
+	DynamicMeshComponent->EditMesh(DeformPreviewMeshFunc, EDynamicMeshComponentRenderUpdateMode::NoUpdate);
+	constexpr bool bNormals = true;
+	constexpr bool bColors = false;
+	constexpr bool bUVs = false;
+	DynamicMeshComponent->FastNotifyPositionsUpdated(bNormals, bColors, bUVs);	
 }
 
 void USkeletalMeshEditingCache::UpdateSelectedBoneIndices()
diff --git a/Engine/Plugins/Animation/SkeletalMeshModelingTools/Source/SkeletalMeshModelingTools/Private/SkeletalMeshEditingCache.h b/Engine/Plugins/Animation/SkeletalMeshModelingTools/Source/SkeletalMeshModelingTools/Private/SkeletalMeshEditingCache.h
--- a/Engine/Plugins/Animation/SkeletalMeshModelingTools/Source/SkeletalMeshModelingTools/Private/SkeletalMeshEditingCache.h
+++ b/Engine/Plugins/Animation/SkeletalMeshModelingTools/Source/SkeletalMeshModelingTools/Private/SkeletalMeshEditingCache.h
@@ -69,8 +69,10 @@ public:
 	TSharedPtr<ISkeletalMeshNotifier> GetNotifier();
 	void HandleNotification(const TArray<FName>& BoneNames, const ESkeletalMeshNotifyType InNotifyType);
 
+	void ToggleForceEnableDynamicMesh(bool bEnable);
 	
-	USkeletalMeshBackedDynamicMeshComponent* GetEditingMeshComponent() const;	
+	USkeletalMeshBackedDynamicMeshComponent* GetEditingMeshComponent() const;
+	UPrimitiveComponent* GetPreviewMeshComponent() const;
 	const TArray<FTransform>& GetComponentSpaceBoneTransforms() const;
 	FTransform GetTransform();
 
@@ -99,6 +101,8 @@ public:
 	FName RenameMorphTarget(FName InOldName, FName InNewName);
 	void RemoveMorphTargets(const TArray<FName>& InNames);
 	TArray<FName> DuplicateMorphTargets(const TArray<FName>& InNames);
+
+	void HandleGeometryUpdate(const UE::Geometry::FDynamicMesh3& InMesh, FName MorphTargetName);
 	
 	URefSkeletonPoser* GetSkeletonPoser() const;
 	
@@ -122,37 +126,50 @@ protected:
 	ESkeletonDrawMode CacheSkeletonDrawMode = ESkeletonDrawMode::Default;
 	bool bCacheBoneManipulation = true;
 
-	bool bEnableDynamicMesh = false;
-	bool bEnableDynamicMeshSkeleton = false;
-	
 	TValueWatcher<int32> ChangeCountWatcher;
+
+	bool bForceEnableDynamicMesh = false;
+	
+	struct FDynamicMeshState
+	{
+		int32 ChangeCount = 0;
+		bool bForceEnable = false;
+		friend bool operator==(const FDynamicMeshState&, const FDynamicMeshState&) = default;	
+	};
+	
+	TValueWatcher<FDynamicMeshState> DynamicMeshStateWatcher;
+	bool bIsDynamicMeshEnabled = false;
+
+
 	TValueWatcher<int32> SkeletonChangeCountWatcher;
+	bool bIsDynamicMeshSkeletonEnabled = false;
 
 	struct FMeshVisibilityState
 	{
-		bool bEnableDynamicMesh= false;
+		bool bIsDynamicMeshEnabled= false;
 		bool bCacheVisibility = true;
 		friend bool operator==(const FMeshVisibilityState&, const FMeshVisibilityState&) = default;
 	};
 
 	TValueWatcher<FMeshVisibilityState> MeshVisibilityWatcher;
-	
+
+
 	struct FSkeletonVisibilityState
 	{
-		bool bEnableDynamicMeshSkeleton = false;
+		bool bIsDynamicMeshSkeletonEnabled = false;
 		ESkeletonDrawMode CacheSkeletonDrawMode = ESkeletonDrawMode::Default;
 		friend bool operator==(const FSkeletonVisibilityState&, const FSkeletonVisibilityState&) = default;
 	};
 	
 	TValueWatcher<FSkeletonVisibilityState> SkeletonVisibilityWatcher;
 
+
 	struct FBoneManipulationState
 	{
-		bool bEnableDynamicMeshSkeleton = false;
+		bool bIsDynamicMeshSkeletonEnabled = false;
 		bool bCacheBoneManipulation = true;
 		friend bool operator==(const FBoneManipulationState&, const FBoneManipulationState&) = default;
 	};
-
 	
 	TValueWatcher<FBoneManipulationState> BoneManipulationWatcher;
 	
diff --git a/Engine/Plugins/Animation/SkeletalMeshModelingTools/Source/SkeletalMeshModelingTools/Private/SkeletalMeshModelingToolsEditorMode.cpp b/Engine/Plugins/Animation/SkeletalMeshModelingTools/Source/SkeletalMeshModelingTools/Private/SkeletalMeshModelingToolsEditorMode.cpp
--- a/Engine/Plugins/Animation/SkeletalMeshModelingTools/Source/SkeletalMeshModelingTools/Private/SkeletalMeshModelingToolsEditorMode.cpp
+++ b/Engine/Plugins/Animation/SkeletalMeshModelingTools/Source/SkeletalMeshModelingTools/Private/SkeletalMeshModelingToolsEditorMode.cpp
@@ -74,6 +74,17 @@
 #include "ToolTargets/SkeletalMeshToolTarget.h"
 #include "ToolTargets/SKMBackedDynaMeshComponentToolTarget.h"
 
+#include "ModelingSelectionInteraction.h"
+#include "ModelingToolsEditorModeSettings.h"
+#include "Selection/DynamicMeshSelector.h"
+#include "Selection/GeometrySelectionManager.h"
+#include "UnrealWidget.h"
+#include "ModelingToolsEditorModeStyle.h"
+#include "ToolMenus.h"
+#include "SkeletalMeshToolMenuContext.h"
+#include "Commands/ModifyGeometrySelectionCommand.h"
+#include "Selection/SKMPreviewDynaMeshSelector.h"
+
 #define LOCTEXT_NAMESPACE "SkeletalMeshModelingToolsEditorMode"
 
 #if ENABLE_STYLUS_SUPPORT
@@ -188,6 +199,12 @@ static void ShowEditorMessage(ELogVerbosity::Type InMessageType, const FText& In
 	FSlateNotificationManager::Get().AddNotification(Notification)->SetCompletionState(State);
 }
 
+namespace SkeletalMeshModelingToolsEditorModeLocals
+{
+	const FName ViewportToolbarExtensionOwnerName = "SkeletalMeshModelingToolsEditorModeViewportToolbar";
+}
+
+
 // NOTE: This is a simple proxy at the moment. In the future we want to pull in more of the 
 // modeling tools as we add support in the skelmesh storage.
 
@@ -361,6 +378,9 @@ void USkeletalMeshModelingToolsEditorMode::Enter()
 	GetInteractiveToolsContext()->ToolManager->OnToolPostBuild.AddUObject(this, &USkeletalMeshModelingToolsEditorMode::OnToolPostBuild);
 	GetInteractiveToolsContext()->ToolManager->OnToolEndedWithStatus.AddUObject(this, &USkeletalMeshModelingToolsEditorMode::OnToolEndedWithStatus);
 
+	InitializeGeometrySelectionSystems();
+	PopulateViewportToolbar();
+	
 	CacheChangeCountWatcher.Initialize(
 		[this]()
 			{
@@ -619,6 +639,26 @@ void USkeletalMeshModelingToolsEditorMode::HandleSetMorphTargetAutoFill(FName Mo
 
 void USkeletalMeshModelingToolsEditorMode::Exit()
 {
+	RemoveViewportToolbarExtensions();
+	
+	// shutdown selection interaction
+	if (SelectionInteraction != nullptr)
+	{
+		SelectionInteraction->Shutdown();
+		GetInteractiveToolsContext()->InputRouter->ForceTerminateSource(SelectionInteraction);
+		GetInteractiveToolsContext()->InputRouter->DeregisterSource(SelectionInteraction);
+		SelectionInteraction = nullptr;
+	}
+	
+	// Shutdown SelectionManager. Wait until after Tool shutdown in case some restore-selection
+	// is involved (although since we are exiting Mode this currently would never matter)
+	if (SelectionManager != nullptr)
+	{
+		SelectionManager->ClearSelection();
+		SelectionManager->Shutdown();		// will clear active targets
+		SelectionManager = nullptr;
+	}
+
 	if (GetCurrentEditingCache())
 	{
 		GetCurrentEditingCache()->ApplyChanges();
@@ -739,6 +779,347 @@ bool USkeletalMeshModelingToolsEditorMode::TryGetExtensionToolCommandGetter(
 	return true;
 }
 
+UGeometrySelectionManager* USkeletalMeshModelingToolsEditorMode::GetSelectionManager()
+{
+	return SelectionManager;
+}
+
+void USkeletalMeshModelingToolsEditorMode::InitializeGeometrySelectionSystems()
+{
+	TWeakObjectPtr<USkeletalMeshModelingToolsEditorMode> WeakThis(this);
+	
+	const FModelingToolsManagerCommands& ToolManagerCommands = FModelingToolsManagerCommands::Get();
+	
+	const TSharedRef<FUICommandList>& CommandList = Toolkit->GetToolkitCommands();
+	
+	// set up SelectionManager and register known factory types
+	SelectionManager = NewObject<UGeometrySelectionManager>(GetToolManager());
+	SelectionManager->Initialize(GetInteractiveToolsContext(), GetToolManager()->GetContextTransactionsAPI());
+	{
+		TUniquePtr<FSkeletalMeshPreviewDynamicMeshSelectorFactory> SelectorFactory =  MakeUnique<FSkeletalMeshPreviewDynamicMeshSelectorFactory>();
+		SelectorFactory->Init(this);
+		SelectionManager->RegisterSelectorFactory(MoveTemp(SelectorFactory));
+	}
+	const UModelingToolsModeCustomizationSettings* ModelingEditorSettings = GetDefault<UModelingToolsModeCustomizationSettings>();
+	
+	// Colors initialized here any time Modeling mode is entered
+	SelectionManager->SetSelectionColors(ModelingEditorSettings->UnselectedColor, ModelingEditorSettings->HoverOverSelectedColor, ModelingEditorSettings->HoverOverUnselectedColor, ModelingEditorSettings->GeometrySelectedColor);
+
+	// This function registers and tracks an active UGeometrySelectionEditCommand and it's associated UICommand
+	auto RegisterSelectionCommand = [this, &CommandList](UGeometrySelectionEditCommand* Command, TSharedPtr<FUICommandInfo> UICommand, bool bAlwaysVisible)
+	{
+		CommandList->MapAction(UICommand,
+			FExecuteAction::CreateUObject(GetSelectionManager(), &UGeometrySelectionManager::ExecuteSelectionCommand, Command),
+			FCanExecuteAction::CreateUObject(GetSelectionManager(), &UGeometrySelectionManager::CanExecuteSelectionCommand, Command),
+			FIsActionChecked(),
+			(bAlwaysVisible) ? FIsActionButtonVisible() : 
+				FIsActionButtonVisible::CreateUObject(GetSelectionManager(), &UGeometrySelectionManager::CanExecuteSelectionCommand, Command),
+			EUIActionRepeatMode::RepeatDisabled);
+	};
+
+	// Create and register InteractiveCommands for mesh selections
+	RegisterSelectionCommand(NewObject<UModifyGeometrySelectionCommand>(), ToolManagerCommands.BeginSelectionAction_SelectAll, true);
+	RegisterSelectionCommand(NewObject<UModifyGeometrySelectionCommand_Invert>(), ToolManagerCommands.BeginSelectionAction_Invert, true);
+	RegisterSelectionCommand(NewObject<UModifyGeometrySelectionCommand_ExpandToConnected>(), ToolManagerCommands.BeginSelectionAction_ExpandToConnected, true);
+	RegisterSelectionCommand(NewObject<UModifyGeometrySelectionCommand_InvertConnected>(), ToolManagerCommands.BeginSelectionAction_InvertConnected, true);
+	RegisterSelectionCommand(NewObject<UModifyGeometrySelectionCommand_Expand>(), ToolManagerCommands.BeginSelectionAction_Expand, true);
+	RegisterSelectionCommand(NewObject<UModifyGeometrySelectionCommand_Contract>(), ToolManagerCommands.BeginSelectionAction_Contract, true);
+
+
+
+	// Set up the selection interaction
+	SelectionInteraction = NewObject<UModelingSelectionInteraction>(GetToolManager());
+	SelectionInteraction->Initialize(SelectionManager,
+		[WeakThis]()
+			{
+				if (!WeakThis.IsValid())
+				{
+					return false;
+				}
+				return WeakThis->bEnableSelection;
+			},
+		[WeakThis](const FInputDeviceRay& DeviceRay)
+			{
+				// Because the editor gizmo does not participate in InputRouter behavior system, in some input behaviors
+				// we need to filter out clicks on the gizmo. This function can do this check.
+				if ( WeakThis.IsValid() && WeakThis->ShouldDrawWidget() )
+				{
+					if (FViewport* FocusedViewport = WeakThis->GetToolManager()->GetContextQueriesAPI()->GetFocusedViewport())
+					{
+						HHitProxy* HitResult = FocusedViewport->GetHitProxy(DeviceRay.ScreenPosition.X, DeviceRay.ScreenPosition.Y);
+						if (HitResult && HitResult->IsA(HWidgetAxis::StaticGetType()))
+						{
+							return true;
+						}
+					}
+				}
+				return false;
+			});
+	GetInteractiveToolsContext()->InputRouter->RegisterSource(SelectionInteraction);
+
+	GetInteractiveToolsContext()->OnRender.AddUObject(this, &USkeletalMeshModelingToolsEditorMode::OnToolContextRender);
+	GetInteractiveToolsContext()->OnDrawHUD.AddUObject(this, &USkeletalMeshModelingToolsEditorMode::OnToolContextDrawHUD);
+
+	// Register Drag Modes
+	auto MakeDragModeOptionAction = [WeakThis](EModelingSelectionInteraction_DragMode DragMode)
+		{
+			return FUIAction(
+				FExecuteAction::CreateLambda([WeakThis, DragMode]
+				{
+					if (!WeakThis.IsValid())
+					{
+						return;
+					}
+						
+					WeakThis->SelectionInteraction->SetActiveDragMode(DragMode);
+					UModelingToolsModeCustomizationSettings* ModelingEditorSettings = GetMutableDefault<UModelingToolsModeCustomizationSettings>();
+					ModelingEditorSettings->LastMeshSelectionDragMode = static_cast<int>(DragMode);
+					ModelingEditorSettings->SaveConfig();
+				}),
+				FCanExecuteAction(),
+				FIsActionChecked::CreateLambda([WeakThis, DragMode]()
+				{
+					if (!WeakThis.IsValid())
+					{
+						return false;
+					}
+
+					return WeakThis->SelectionInteraction->GetActiveDragMode() == DragMode;
+				}));
+		};
+
+	// Map drag mode commands
+	CommandList->MapAction(ToolManagerCommands.SelectionDragMode_None,
+		MakeDragModeOptionAction(EModelingSelectionInteraction_DragMode::NoDragInteraction));
+	CommandList->MapAction(ToolManagerCommands.SelectionDragMode_Path,
+		MakeDragModeOptionAction(EModelingSelectionInteraction_DragMode::PathInteraction));
+
+	// Register Frame Modes
+	auto ToggleFromGeometryAction = [WeakThis](EModelingSelectionInteraction_LocalFrameMode LocalFrameMode)
+		{
+			return FUIAction(
+				FExecuteAction::CreateLambda([WeakThis, LocalFrameMode]
+				{
+					if (!WeakThis.IsValid())
+					{
+						return;
+					}
+					WeakThis->SelectionInteraction->SetLocalFrameMode(LocalFrameMode);
+					UModelingToolsModeCustomizationSettings* ModelingEditorSettings = GetMutableDefault<UModelingToolsModeCustomizationSettings>();
+					ModelingEditorSettings->LastMeshSelectionLocalFrameMode = static_cast<int>(LocalFrameMode);
+					ModelingEditorSettings->SaveConfig();
+				}),
+				FCanExecuteAction(),
+				FIsActionChecked::CreateLambda([WeakThis, LocalFrameMode]()
+				{
+					if (!WeakThis.IsValid())
+					{
+						return false;
+					}
+					return WeakThis->SelectionInteraction->GetLocalFrameMode() == LocalFrameMode;
+				}));
+		};
+	
+	CommandList->MapAction(ToolManagerCommands.SelectionLocalFrameMode_Geometry,
+		ToggleFromGeometryAction(EModelingSelectionInteraction_LocalFrameMode::FromGeometry));
+	CommandList->MapAction(ToolManagerCommands.SelectionLocalFrameMode_Object,
+		ToggleFromGeometryAction( EModelingSelectionInteraction_LocalFrameMode::FromObject));
+	
+	// Register Selection Hit Back Faces Toggle
+	FUIAction ToggleHitBackFacesAction = FUIAction(
+		FExecuteAction::CreateLambda([WeakThis]
+		{
+			UModelingToolsModeCustomizationSettings* ModelingEditorSettings = GetMutableDefault<UModelingToolsModeCustomizationSettings>();
+			ModelingEditorSettings->bMeshSelectionHitBackFaces = !ModelingEditorSettings->bMeshSelectionHitBackFaces;
+			WeakThis->SelectionManager->SetHitBackFaces(ModelingEditorSettings->bMeshSelectionHitBackFaces);
+			ModelingEditorSettings->SaveConfig();
+		}),
+		FCanExecuteAction(),
+		FIsActionChecked::CreateLambda([WeakThis]()
+		{
+			if (!WeakThis.IsValid())
+			{
+				return false;
+			}
+			return WeakThis->SelectionManager->GetHitBackFaces();
+		}));
+	CommandList->MapAction(ToolManagerCommands.SelectionHitBackFaces, ToggleHitBackFacesAction);
+
+	
+	// set up selection type toggles
+	auto RegisterSelectionMode = [WeakThis](UGeometrySelectionManager::EMeshTopologyMode TopoMode, UE::Geometry::EGeometryElementType ElementMode, TSharedPtr<FUICommandInfo> UICommand)
+	{
+		WeakThis->Toolkit->GetToolkitCommands()->MapAction(UICommand,
+			FExecuteAction::CreateLambda([WeakThis, TopoMode, ElementMode]() {
+				if ( WeakThis.IsValid() && WeakThis->GetToolManager() && WeakThis->GetToolManager()->GetContextTransactionsAPI())
+				{
+					WeakThis->GetToolManager()->GetContextTransactionsAPI()->BeginUndoTransaction(LOCTEXT("ChangeSelectionMode", "Selection Mode"));
+					WeakThis->SelectionManager->SetMeshSelectionTypeAndMode(ElementMode, TopoMode,  TopoMode != UGeometrySelectionManager::EMeshTopologyMode::None);
+					WeakThis->GetToolManager()->GetContextTransactionsAPI()->EndUndoTransaction();
+					if (UModelingToolsModeCustomizationSettings* ModelingEditorSettings = GetMutableDefault<UModelingToolsModeCustomizationSettings>())
+					{
+						ModelingEditorSettings->LastMeshSelectionTopologyMode = static_cast<int>(TopoMode);
+						ModelingEditorSettings->LastMeshSelectionElementType = static_cast<int>(ElementMode);
+						ModelingEditorSettings->SaveConfig();
+					}
+					
+					if (WeakThis->GetCurrentEditingCache())
+					{
+						const bool bElementSelectionEnabled = TopoMode != UGeometrySelectionManager::EMeshTopologyMode::None;
+						WeakThis->GetCurrentEditingCache()->ToggleForceEnableDynamicMesh(bElementSelectionEnabled);
+					}
+				}
+			}),
+			FCanExecuteAction::CreateLambda([WeakThis]()
+			{
+				if (!WeakThis.IsValid())
+				{
+					return false;
+				}
+				return WeakThis->bEnableSelection;
+			}),
+			FIsActionChecked::CreateLambda([WeakThis, TopoMode, ElementMode]()
+				{
+					if (!WeakThis.IsValid() || !WeakThis->SelectionManager)
+					{
+						return false;
+					}
+					return WeakThis->SelectionManager->GetMeshTopologyMode() == TopoMode && WeakThis->SelectionManager->GetSelectionElementType() == ElementMode;
+				}),
+			EUIActionRepeatMode::RepeatDisabled);
+	};
+	{
+		RegisterSelectionMode(UGeometrySelectionManager::EMeshTopologyMode::None, UGeometrySelectionManager::EGeometryElementType::Face, ToolManagerCommands.MeshSelectionModeAction_NoSelection);
+		RegisterSelectionMode(UGeometrySelectionManager::EMeshTopologyMode::Triangle, UGeometrySelectionManager::EGeometryElementType::Face, ToolManagerCommands.MeshSelectionModeAction_MeshTriangles);
+		RegisterSelectionMode(UGeometrySelectionManager::EMeshTopologyMode::Triangle, UGeometrySelectionManager::EGeometryElementType::Vertex, ToolManagerCommands.MeshSelectionModeAction_MeshVertices);
+		RegisterSelectionMode(UGeometrySelectionManager::EMeshTopologyMode::Triangle, UGeometrySelectionManager::EGeometryElementType::Edge, ToolManagerCommands.MeshSelectionModeAction_MeshEdges);
+		RegisterSelectionMode(UGeometrySelectionManager::EMeshTopologyMode::Polygroup, UGeometrySelectionManager::EGeometryElementType::Face, ToolManagerCommands.MeshSelectionModeAction_GroupFaces);
+		RegisterSelectionMode(UGeometrySelectionManager::EMeshTopologyMode::Polygroup, UGeometrySelectionManager::EGeometryElementType::Vertex, ToolManagerCommands.MeshSelectionModeAction_GroupCorners);
+		RegisterSelectionMode(UGeometrySelectionManager::EMeshTopologyMode::Polygroup, UGeometrySelectionManager::EGeometryElementType::Edge, ToolManagerCommands.MeshSelectionModeAction_GroupEdges);
+	}
+}
+
+void USkeletalMeshModelingToolsEditorMode::PopulateViewportToolbar()
+{
+	using namespace SkeletalMeshModelingToolsEditorModeLocals;
+	FToolMenuOwnerScoped ScopeOwner(ViewportToolbarExtensionOwnerName);
+	UToolMenu* const Menu = UToolMenus::Get()->ExtendMenu("AnimationEditor.ViewportToolbar");
+	
+	
+	FToolMenuSection& MeshElementSelectionMenu = Menu->FindOrAddSection("Left");
+	FToolMenuEntry Entry = FToolMenuEntry::InitSubMenu(
+		"Mesh Element Selection",
+		LOCTEXT("MeshElementSelectionSubmenuLabel", "Mesh Element Selection"),
+		LOCTEXT("MeshElementSelectionSubmenuTooltip", "Mesh Element Selection settings in the viewport"),
+		FNewToolMenuDelegate::CreateLambda([ WeakCommandList = Toolkit->GetToolkitCommands().ToWeakPtr()](UToolMenu* Submenu) -> void
+		{
+			FToolMenuSection& MeshElementSelectionOptionsSection =
+			Submenu->FindOrAddSection("Element Selection", LOCTEXT("ElementSelectionLabel", "Element Selection"));
+			if (const TSharedPtr<const FUICommandList>& CommandList = WeakCommandList.Pin())
+			{
+				auto CreateToolMenuEntry = [&CommandList, &MeshElementSelectionOptionsSection](const TSharedPtr<FUICommandInfo>& Command)
+				{
+					FToolMenuEntry MESMode = FToolMenuEntry::InitMenuEntryWithCommandList( Command, CommandList );
+					MESMode.SetShowInToolbarTopLevel(true); // add button to toolbar
+					MeshElementSelectionOptionsSection.AddEntry(MESMode);
+				};
+
+				// add all Mesh Element Selection modes
+				CreateToolMenuEntry(FModelingToolsManagerCommands::Get().MeshSelectionModeAction_NoSelection);
+				MeshElementSelectionOptionsSection.AddSeparator(NAME_None);
+				CreateToolMenuEntry(FModelingToolsManagerCommands::Get().MeshSelectionModeAction_MeshVertices);
+				CreateToolMenuEntry(FModelingToolsManagerCommands::Get().MeshSelectionModeAction_MeshEdges);
+				CreateToolMenuEntry(FModelingToolsManagerCommands::Get().MeshSelectionModeAction_MeshTriangles);
+				MeshElementSelectionOptionsSection.AddSeparator(NAME_None);
+				CreateToolMenuEntry(FModelingToolsManagerCommands::Get().MeshSelectionModeAction_GroupCorners);
+				CreateToolMenuEntry(FModelingToolsManagerCommands::Get().MeshSelectionModeAction_GroupEdges);
+				CreateToolMenuEntry(FModelingToolsManagerCommands::Get().MeshSelectionModeAction_GroupFaces);
+
+				// Section for Selection Edits
+				{
+					FToolMenuSection& SelectionEditsSection =
+						Submenu->FindOrAddSection("SelectionEdits", LOCTEXT("SelectionEditsLabel", "Selection Edits"));
+
+					SelectionEditsSection.AddEntry(
+					FToolMenuEntry::InitMenuEntryWithCommandList(
+						FModelingToolsManagerCommands::Get().BeginSelectionAction_SelectAll, CommandList));
+					SelectionEditsSection.AddEntry(
+					FToolMenuEntry::InitMenuEntryWithCommandList(
+						FModelingToolsManagerCommands::Get().BeginSelectionAction_ExpandToConnected, CommandList));
+					SelectionEditsSection.AddEntry(
+					FToolMenuEntry::InitMenuEntryWithCommandList(
+						FModelingToolsManagerCommands::Get().BeginSelectionAction_Invert, CommandList));
+					SelectionEditsSection.AddEntry(
+					FToolMenuEntry::InitMenuEntryWithCommandList(
+						FModelingToolsManagerCommands::Get().BeginSelectionAction_InvertConnected, CommandList));
+					SelectionEditsSection.AddEntry(
+					FToolMenuEntry::InitMenuEntryWithCommandList(
+						FModelingToolsManagerCommands::Get().BeginSelectionAction_Expand, CommandList));
+					SelectionEditsSection.AddEntry(
+					FToolMenuEntry::InitMenuEntryWithCommandList(
+						FModelingToolsManagerCommands::Get().BeginSelectionAction_Contract, CommandList));
+				}
+
+				// Section for Drag Mode setting
+				{
+					FToolMenuSection& DragModeSection =
+						Submenu->FindOrAddSection("Drag Mode", LOCTEXT("DragModeLabel", "Drag Mode"));
+					DragModeSection.AddEntry(
+						FToolMenuEntry::InitMenuEntryWithCommandList(
+							FModelingToolsManagerCommands::Get().SelectionDragMode_None, CommandList));
+					DragModeSection.AddEntry(
+						FToolMenuEntry::InitMenuEntryWithCommandList(
+							FModelingToolsManagerCommands::Get().SelectionDragMode_Path, CommandList));
+				}
+
+				
+				// Section for Local Frame Mode setting
+				{
+					FToolMenuSection& LocalFrameModeSection =
+						Submenu->FindOrAddSection("Local Frame Mode", LOCTEXT("LocalFrameModeLabel", "Local Frame Mode"));
+					LocalFrameModeSection.AddEntry(
+						FToolMenuEntry::InitMenuEntryWithCommandList(
+							FModelingToolsManagerCommands::Get().SelectionLocalFrameMode_Geometry, CommandList));
+					LocalFrameModeSection.AddEntry(
+						FToolMenuEntry::InitMenuEntryWithCommandList(
+							FModelingToolsManagerCommands::Get().SelectionLocalFrameMode_Object, CommandList));
+				}
+
+				// Section for Selection Filters
+				{
+					FToolMenuSection& SelectionFilterSection =
+						Submenu->FindOrAddSection("Selection Filters", LOCTEXT("SelectionFiltersLabel", "Selection Filters"));
+					SelectionFilterSection.AddEntry(
+							FToolMenuEntry::InitMenuEntryWithCommandList(
+								FModelingToolsManagerCommands::Get().SelectionHitBackFaces, CommandList));
+				}
+
+			}
+		})
+	);
+
+	Entry.ToolBarData.ResizeParams.ClippingPriority = 950;
+	Entry.ToolBarData.LabelOverride = FText();
+	Entry.InsertPosition = FToolMenuInsert("Transform", EToolMenuInsertType::After);
+	Entry.Icon = FSlateIcon(FModelingToolsEditorModeStyle::GetStyleSetName(), "ModelingToolsManagerCommands.MeshElementSelection");
+
+	MeshElementSelectionMenu.AddDynamicEntry(TEXT("Mesh Element Selection"),
+		FNewToolMenuSectionDelegate::CreateLambda([Entry](FToolMenuSection& InSection)
+		{
+			// Only extending the menu for skeletal mesh editor
+			const USkeletalMeshToolMenuContext* Context = InSection.FindContext<USkeletalMeshToolMenuContext>();
+			if (Context && Context->SkeletalMeshEditor.IsValid())
+			{
+				InSection.AddEntry(Entry);
+			}
+		}));
+}
+
+void USkeletalMeshModelingToolsEditorMode::RemoveViewportToolbarExtensions()
+{
+	UToolMenus::Get()->UnregisterOwnerByName(SkeletalMeshModelingToolsEditorModeLocals::ViewportToolbarExtensionOwnerName);
+}
+
 void USkeletalMeshModelingToolsEditorMode::CreateToolkit()
 {
 	TSharedPtr<FSkeletalMeshModelingToolsEditorModeToolkit> NewToolKit = MakeShareable(new FSkeletalMeshModelingToolsEditorModeToolkit);
@@ -882,6 +1263,9 @@ void USkeletalMeshModelingToolsEditorMode::OnToolPostBuild( UInteractiveToolMana
 {
 	// Tools can decide for themselves if they want to enable bone manipulation
 	ToggleBoneManipulation(false);
+
+	// Disable element selection modes while tool is active
+	bEnableSelection = false;
 }
 
 void USkeletalMeshModelingToolsEditorMode::OnToolEndedWithStatus(UInteractiveToolManager* Manager, UInteractiveTool* Tool,
@@ -895,6 +1279,30 @@ void USkeletalMeshModelingToolsEditorMode::OnToolEndedWithStatus(UInteractiveToo
 	}
 	
 	ToggleBoneManipulation(true);
+
+	bEnableSelection = true;
+}
+
+void USkeletalMeshModelingToolsEditorMode::OnToolContextRender(IToolsContextRenderAPI* ToolsContextRenderAPI)
+{
+	if (SelectionInteraction)
+	{
+		SelectionInteraction->Render(ToolsContextRenderAPI);
+
+		// Bake in transform changes. Note that if we do this in OnToolsContextTick, it will 
+		// still block rendering updates if it is too expensive, unless it is only done every second Tick
+		SelectionInteraction->ApplyPendingTransformInteractions();
+	}
+	
+	SelectionManager->DebugRender(ToolsContextRenderAPI);
+}
+
+void USkeletalMeshModelingToolsEditorMode::OnToolContextDrawHUD(FCanvas* Canvas, IToolsContextRenderAPI* ToolsContextRenderAPI)
+{
+	if (SelectionInteraction)
+	{
+		SelectionInteraction->DrawHUD(Canvas, ToolsContextRenderAPI);
+	}	
 }
 
 bool USkeletalMeshModelingToolsEditorMode::ShouldToolStartBeAllowed(const FString& ToolIdentifier) const
@@ -1281,15 +1689,20 @@ void USkeletalMeshModelingToolsEditorMode::RecreateEditingCache(EMeshLODIdentifi
 	
 	CurrentEditingCache = NewObject<USkeletalMeshEditingCache>();
 
-	TWeakObjectPtr<USkeletalMeshModelingToolsEditorMode> WeakThis(this);
 	
 	USkeletalMeshEditingCache::FDelegates Delegates;
 	Delegates.ToggleSkeletalMeshBoneManipulationDelegate.BindUObject(this, &USkeletalMeshModelingToolsEditorMode::ToggleSkeletalMeshBoneManipulation);
 	Delegates.IsSkeletalMeshBoneManipulationEnabledDelegate.BindUObject(this, &USkeletalMeshModelingToolsEditorMode::IsSkeletalMeshBoneManipulationEnabled);
-	Delegates.OnGetSkeletalMeshSkeletonNotifierDelegate.BindLambda([this]()
+	
+	Delegates.OnGetSkeletalMeshSkeletonNotifierDelegate.BindLambda([WeakThis = TWeakObjectPtr<USkeletalMeshModelingToolsEditorMode>(this)]()
 		{
 			// Skeletal Mesh Skeleton bone selection state is governed by the skeleton tree in the base editor
-			return GetEditorBinding()->GetNotifier();
+			if (WeakThis.IsValid())
+			{
+				return WeakThis->GetEditorBinding()->GetNotifier();
+			}
+
+			return TSharedPtr<ISkeletalMeshNotifier>();
 		});
 	CurrentEditingCache->Spawn(PreviewWorld, GetSkelMeshComponent(), InLOD, Delegates);
 
@@ -1299,6 +1712,9 @@ void USkeletalMeshModelingToolsEditorMode::RecreateEditingCache(EMeshLODIdentifi
 
 	CurrentEditingCache->GetNotifier()->HandleNotification({}, ESkeletalMeshNotifyType::HierarchyChanged);
 	CurrentEditingCache->GetNotifier()->HandleNotification(GetSelectedBones(), ESkeletalMeshNotifyType::BonesSelected);
+	
+	SelectionManager->ClearActiveTargets();
+	SelectionManager->AddActiveTarget(FGeometryIdentifier::PrimitiveComponent(CurrentEditingCache->GetPreviewMeshComponent(), FGeometryIdentifier::EObjectType::DynamicMeshComponent));
 }
 
 
diff --git a/Engine/Plugins/Animation/SkeletalMeshModelingTools/Source/SkeletalMeshModelingTools/Private/SkeletalMeshModelingToolsEditorMode.h b/Engine/Plugins/Animation/SkeletalMeshModelingTools/Source/SkeletalMeshModelingTools/Private/SkeletalMeshModelingToolsEditorMode.h
--- a/Engine/Plugins/Animation/SkeletalMeshModelingTools/Source/SkeletalMeshModelingTools/Private/SkeletalMeshModelingToolsEditorMode.h
+++ b/Engine/Plugins/Animation/SkeletalMeshModelingTools/Source/SkeletalMeshModelingTools/Private/SkeletalMeshModelingToolsEditorMode.h
@@ -6,6 +6,7 @@
 #include "EdMode.h"
 #include "SkeletalMeshModelingModeToolExtensions.h"
 #include "SkeletalMeshNotifier.h"
+#include "ToolContextInterfaces.h"
 #include "Interfaces/ISKMBackedDynaMeshComponentProvider.h"
 #include "Animation/DebugSkelMeshComponent.h"
 #include "Changes/ValueWatcher.h"
@@ -31,6 +32,8 @@ enum class EToolSide;
 enum class EToolShutdownType : uint8;
 struct FToolBuilderState;
 class USkeletonModifier;
+class UGeometrySelectionManager;
+class UModelingSelectionInteraction;
 
 namespace UE
 {
@@ -172,7 +175,9 @@ protected:
 	virtual void OnToolStarted(UInteractiveToolManager* Manager, UInteractiveTool* Tool) override;
 	void OnToolPostBuild(UInteractiveToolManager* InToolManager, EToolSide InSide, UInteractiveTool* InBuiltTool, UInteractiveToolBuilder* InToolBuilder, const FToolBuilderState& ToolState);
 	void OnToolEndedWithStatus(UInteractiveToolManager* Manager, UInteractiveTool* Tool, EToolShutdownType ShutdownType);
-
+	
+	void OnToolContextRender(IToolsContextRenderAPI* ToolsContextRenderAPI);
+	void OnToolContextDrawHUD(FCanvas* Canvas, IToolsContextRenderAPI* ToolsContextRenderAPI);
 private:
 	EApplyMode ApplyMode = EApplyMode::ApplyOnToolExit;
 	
@@ -244,4 +249,18 @@ private:
 	TFunction<const UE::IInteractiveToolCommandsInterface& ()> ExtensionToolCommandsGetter;
 
 	FSimpleMulticastDelegate OnInitializedDelegate;
+
+	UGeometrySelectionManager* GetSelectionManager();
+
+	UPROPERTY(NonTransactional)
+	TObjectPtr<UGeometrySelectionManager> SelectionManager;
+
+	UPROPERTY(NonTransactional)
+	TObjectPtr<UModelingSelectionInteraction> SelectionInteraction;
+
+	bool bEnableSelection = true;
+
+	void InitializeGeometrySelectionSystems();
+	void PopulateViewportToolbar();
+	void RemoveViewportToolbarExtensions();
 };
diff --git a/Engine/Plugins/Editor/ModelingToolsEditorMode/Source/ModelingToolsEditorMode/Public/ModelingSelectionInteraction.h b/Engine/Plugins/Editor/ModelingToolsEditorMode/Source/ModelingToolsEditorMode/Public/ModelingSelectionInteraction.h
--- a/Engine/Plugins/Editor/ModelingToolsEditorMode/Source/ModelingToolsEditorMode/Public/ModelingSelectionInteraction.h
+++ b/Engine/Plugins/Editor/ModelingToolsEditorMode/Source/ModelingToolsEditorMode/Public/ModelingSelectionInteraction.h
@@ -9,6 +9,8 @@
 #include "UObject/Package.h"
 #include "ModelingSelectionInteraction.generated.h"
 
+#define UE_API MODELINGTOOLSEDITORMODE_API
+
 class UCombinedTransformGizmo;
 class URectangleMarqueeInteraction;
 class UTransformProxy;
@@ -60,7 +62,7 @@ enum class EModelingSelectionInteraction_LocalFrameMode : uint8
  * Editor interaction behaviors are re-implemented as InputBehaviors)
  * 
  */
-UCLASS()
+UCLASS(MinimalAPI)
 class UModelingSelectionInteraction : public UObject, public IInputBehaviorSource, public IClickBehaviorTarget, public IHoverBehaviorTarget
 {
 	GENERATED_BODY()
@@ -90,7 +92,7 @@ public:
 	//
 
 	EModelingSelectionInteraction_DragMode GetActiveDragMode() const { return ActiveDragMode; }
-	void SetActiveDragMode(EModelingSelectionInteraction_DragMode NewMode);
+	UE_API void SetActiveDragMode(EModelingSelectionInteraction_DragMode NewMode);
 
 	//
 	// Support for transformation of selected geometry
@@ -109,7 +111,7 @@ public:
 	// Support for Adjusting Local Frame Mode
 	//
 	EModelingSelectionInteraction_LocalFrameMode GetLocalFrameMode() const { return LocalFrameMode; }
-	void SetLocalFrameMode(EModelingSelectionInteraction_LocalFrameMode NewLocalFrameMode);
+	UE_API void SetLocalFrameMode(EModelingSelectionInteraction_LocalFrameMode NewLocalFrameMode);
 
 protected:
 	// click-to-select behavior, various drag behaviors
@@ -256,3 +258,5 @@ protected:
 	virtual void OnTerminateDragSequence() override;
 
 };
+
+#undef UE_API
\ No newline at end of file
diff --git a/Engine/Plugins/Runtime/MeshModelingToolset/Source/MeshModelingToolsEditorOnly/Private/SkeletalMesh/SkeletalMeshToolsHelper.cpp b/Engine/Plugins/Runtime/MeshModelingToolset/Source/MeshModelingToolsEditorOnly/Private/SkeletalMesh/SkeletalMeshToolsHelper.cpp
--- a/Engine/Plugins/Runtime/MeshModelingToolset/Source/MeshModelingToolsEditorOnly/Private/SkeletalMesh/SkeletalMeshToolsHelper.cpp
+++ b/Engine/Plugins/Runtime/MeshModelingToolset/Source/MeshModelingToolsEditorOnly/Private/SkeletalMesh/SkeletalMeshToolsHelper.cpp
@@ -59,10 +59,13 @@ void SkeletalMeshToolsHelper::GetUnposedMesh(
 					const FDynamicMeshMorphTargetAttribute* MorphTargetAttribute =
 						SourceMesh.Attributes()->GetMorphTargetAttribute(MorphTargetWeight.Key);
 
-					FVector Delta;
-					MorphTargetAttribute->GetValue(VertID, Delta);
+					if (ensure(MorphTargetAttribute))
+					{
+						FVector Delta;
+						MorphTargetAttribute->GetValue(VertID, Delta);
 				
-					UnposedVertPos += MorphTargetWeight.Value * Delta * (-1.0f);
+						UnposedVertPos += MorphTargetWeight.Value * Delta * (-1.0f);
+					}
 				}
 			}
 
@@ -102,10 +105,13 @@ void SkeletalMeshToolsHelper::GetPosedMesh(
 					const FDynamicMeshMorphTargetAttribute* MorphTargetAttribute =
 						SourceMesh.Attributes()->GetMorphTargetAttribute(MorphTargetWeight.Key);
 
-					FVector Delta;
-					MorphTargetAttribute->GetValue(VertID, Delta);
+					if (ensure(MorphTargetAttribute))
+					{
+						FVector Delta;
+						MorphTargetAttribute->GetValue(VertID, Delta);
 			
-					PosedVertPos += MorphTargetWeight.Value * Delta;
+						PosedVertPos += MorphTargetWeight.Value * Delta;
+					}
 				}
 			}	
 				
