[SkeletalMeshModellingTools] adding referencer property for geo commands to avoid them getting GCed

#Synced-CL 48737829

diff --git a/Engine/Plugins/Animation/SkeletalMeshModelingTools/Source/SkeletalMeshModelingTools/Private/SkeletalMeshModelingToolsEditorMode.cpp b/Engine/Plugins/Animation/SkeletalMeshModelingTools/Source/SkeletalMeshModelingTools/Private/SkeletalMeshModelingToolsEditorMode.cpp
--- a/Engine/Plugins/Animation/SkeletalMeshModelingTools/Source/SkeletalMeshModelingTools/Private/SkeletalMeshModelingToolsEditorMode.cpp
+++ b/Engine/Plugins/Animation/SkeletalMeshModelingTools/Source/SkeletalMeshModelingTools/Private/SkeletalMeshModelingToolsEditorMode.cpp
@@ -808,6 +808,7 @@ void USkeletalMeshModelingToolsEditorMode::InitializeGeometrySelectionSystems()
 	// This function registers and tracks an active UGeometrySelectionEditCommand and it's associated UICommand
 	auto RegisterSelectionCommand = [this, &CommandList](UGeometrySelectionEditCommand* Command, TSharedPtr<FUICommandInfo> UICommand, bool bAlwaysVisible)
 	{
+		ModelingModeCommands.Add(Command);
 		CommandList->MapAction(UICommand,
 			FExecuteAction::CreateUObject(GetSelectionManager(), &UGeometrySelectionManager::ExecuteSelectionCommand, Command),
 			FCanExecuteAction::CreateUObject(GetSelectionManager(), &UGeometrySelectionManager::CanExecuteSelectionCommand, Command),
diff --git a/Engine/Plugins/Animation/SkeletalMeshModelingTools/Source/SkeletalMeshModelingTools/Private/SkeletalMeshModelingToolsEditorMode.h b/Engine/Plugins/Animation/SkeletalMeshModelingTools/Source/SkeletalMeshModelingTools/Private/SkeletalMeshModelingToolsEditorMode.h
--- a/Engine/Plugins/Animation/SkeletalMeshModelingTools/Source/SkeletalMeshModelingTools/Private/SkeletalMeshModelingToolsEditorMode.h
+++ b/Engine/Plugins/Animation/SkeletalMeshModelingTools/Source/SkeletalMeshModelingTools/Private/SkeletalMeshModelingToolsEditorMode.h
@@ -34,6 +34,7 @@ struct FToolBuilderState;
 class USkeletonModifier;
 class UGeometrySelectionManager;
 class UModelingSelectionInteraction;
+class UInteractiveCommand;
 
 namespace UE
 {
@@ -263,4 +264,7 @@ private:
 	void InitializeGeometrySelectionSystems();
 	void PopulateViewportToolbar();
 	void RemoveViewportToolbarExtensions();
+	
+	UPROPERTY(NonTransactional)
+	TArray<TObjectPtr<UInteractiveCommand>> ModelingModeCommands;
 };
