[SkeletalMeshModelingTools] Using a more reliable way to access the preview mesh component

#rnx
#Synced-CL 49562189

diff --git a/Engine/Plugins/Animation/SkeletalMeshModelingTools/Source/SkeletalMeshModelingTools/Private/SkeletalMeshModelingToolsEditorMode.cpp b/Engine/Plugins/Animation/SkeletalMeshModelingTools/Source/SkeletalMeshModelingTools/Private/SkeletalMeshModelingToolsEditorMode.cpp
--- a/Engine/Plugins/Animation/SkeletalMeshModelingTools/Source/SkeletalMeshModelingTools/Private/SkeletalMeshModelingToolsEditorMode.cpp
+++ b/Engine/Plugins/Animation/SkeletalMeshModelingTools/Source/SkeletalMeshModelingTools/Private/SkeletalMeshModelingToolsEditorMode.cpp
@@ -467,13 +467,13 @@ void USkeletalMeshModelingToolsEditorMode::Enter()
 
 UDebugSkelMeshComponent* USkeletalMeshModelingToolsEditorMode::GetSkelMeshComponent() const
 {
-	FToolBuilderState State; GetToolManager()->GetContextQueriesAPI()->GetCurrentSelectionState(State);
-	UActorComponent* SkeletalMeshComponent = ToolBuilderUtil::FindFirstComponent(State, [&](UActorComponent* Component)
+	if (!Editor.IsValid())
 	{
-		return IsValid(Component) && Component->IsA<UDebugSkelMeshComponent>();
-	});
-
-	return Cast<UDebugSkelMeshComponent>(SkeletalMeshComponent);
+		return nullptr;
+	}
+	
+	TSharedRef<class IPersonaToolkit> PersonaToolKit = Editor.Pin()->GetPersonaToolkit();
+	return PersonaToolKit->GetPreviewMeshComponent();
 }
 
 
