[SkeletalMeshModelingTools] Fixed a minor ensure related to stale morph target weight overrides

#rnx

#Synced-CL 49562189

diff --git a/Engine/Plugins/Animation/SkeletalMeshModelingTools/Source/SkeletalMeshModelingTools/Private/SkeletalMeshEditingCache.cpp b/Engine/Plugins/Animation/SkeletalMeshModelingTools/Source/SkeletalMeshModelingTools/Private/SkeletalMeshEditingCache.cpp
--- a/Engine/Plugins/Animation/SkeletalMeshModelingTools/Source/SkeletalMeshModelingTools/Private/SkeletalMeshEditingCache.cpp
+++ b/Engine/Plugins/Animation/SkeletalMeshModelingTools/Source/SkeletalMeshModelingTools/Private/SkeletalMeshEditingCache.cpp
@@ -595,7 +595,13 @@ TMap<FName, float> USkeletalMeshEditingCache::GetMorphTargetWeights() const
 		}
 	}
 
-	MorphTargetWeights.Append(MorphTargetOverrides);
+	for (TPair<FName, float>& Entry : MorphTargetWeights)
+	{
+		if (const float* Override = MorphTargetOverrides.Find(Entry.Key))
+		{
+			Entry.Value = *Override;
+		}
+	}
 
 	return MorphTargetWeights;	
 }
