[SkeletalMeshModellingTools] Mirror and Orient no longer show an obsolete none axis option

#Synced-CL 49723230

diff --git a/Engine/Plugins/Animation/SkeletalMeshModelingTools/Source/SkeletalMeshModelingTools/Private/DetailCustomization/SKMMTMirroringPropertiesAxisCustomization.cpp b/Engine/Plugins/Animation/SkeletalMeshModelingTools/Source/SkeletalMeshModelingTools/Private/DetailCustomization/SKMMTMirroringPropertiesAxisCustomization.cpp
new file mode 100644
--- /dev/null
+++ b/Engine/Plugins/Animation/SkeletalMeshModelingTools/Source/SkeletalMeshModelingTools/Private/DetailCustomization/SKMMTMirroringPropertiesAxisCustomization.cpp
@@ -0,0 +1,95 @@
+// Copyright Epic Games, Inc. All Rights Reserved.
+
+#include "SKMMTMirroringPropertiesAxisCustomization.h"
+
+#include "DetailWidgetRow.h"
+#include "IDetailChildrenBuilder.h"
+#include "PropertyHandle.h"
+#include "SkeletalMesh/SkeletonEditingTool.h"
+#include "Widgets/SSKMEnumPropertyComboBox.h"
+
+namespace UE::SkeletalMeshModelingTools
+{
+	TSharedRef<IPropertyTypeCustomization> FSKMMTMirroringPropertiesAxisCustomization::MakeInstance()
+	{
+		return MakeShared<FSKMMTMirroringPropertiesAxisCustomization>();
+	}
+
+	void FSKMMTMirroringPropertiesAxisCustomization::CustomizeHeader(
+		TSharedRef<IPropertyHandle> InStructPropertyHandle,
+		FDetailWidgetRow& HeaderRow,
+		IPropertyTypeCustomizationUtils& StructCustomizationUtils)
+	{
+	}
+
+	void FSKMMTMirroringPropertiesAxisCustomization::CustomizeChildren(
+		TSharedRef<IPropertyHandle> InStructPropertyHandle, 
+		IDetailChildrenBuilder& StructBuilder, 
+		IPropertyTypeCustomizationUtils& StructCustomizationUtils)
+	{
+		TArray<UObject*> OuterObjects;
+		InStructPropertyHandle->GetOuterObjects(OuterObjects);
+
+		const bool bValid = OuterObjects.Num() == 1 && OuterObjects[0] && OuterObjects[0]->GetClass() == UMirroringProperties::StaticClass();
+		if (!ensureMsgf(bValid, TEXT("Unexpected unexpected objects wheen trying to customize the Mirroring Properties' Mirroring Axis")))
+		{
+			return;
+		}
+		const TWeakObjectPtr<UMirroringProperties> WeakMirroringProperties = CastChecked<UMirroringProperties>(OuterObjects[0]);
+
+		const TAttribute<EVisibility> VisibilityAttribute = TAttribute<EVisibility>::CreateLambda([WeakMirroringProperties]()
+			{
+				const bool bVisible =
+					WeakMirroringProperties.IsValid() &&
+					WeakMirroringProperties->ParentTool.IsValid() &&
+					WeakMirroringProperties->ParentTool->GetOperation() != EEditingOperation::Create;
+
+				return bVisible ? EVisibility::Visible : EVisibility::Hidden;
+			});
+
+		const TAttribute<bool> EnabledAttribute = TAttribute<bool>::CreateLambda([WeakMirroringProperties]()
+			{
+				return
+					WeakMirroringProperties.IsValid() &&
+					WeakMirroringProperties->ParentTool.IsValid() &&
+					!WeakMirroringProperties->ParentTool->GetSelection().IsEmpty() &&
+					WeakMirroringProperties->ParentTool->GetOperation() != EEditingOperation::Parent;
+			});
+
+		uint32 NumChildren = 0;
+		InStructPropertyHandle->GetNumChildren(NumChildren);
+
+		for (uint32 ChildIndex = 0; ChildIndex < NumChildren; ChildIndex++)
+		{
+			if (const TSharedPtr<IPropertyHandle> ChildProperty = InStructPropertyHandle->GetChildHandle(ChildIndex))
+			{
+				const FProperty* Property = ChildProperty->GetProperty();
+				const FName PropertyName = Property ? Property->GetFName() : NAME_None;
+
+				if (PropertyName == GET_MEMBER_NAME_CHECKED(FMirrorOptions, MirrorAxis))
+				{	
+					StructBuilder.AddCustomRow(ChildProperty->GetPropertyDisplayName())
+						.NameContent()
+						[
+							ChildProperty->CreatePropertyNameWidget()
+						]
+						.ValueContent()
+						[
+							SNew(SSKMEnumPropertyComboBox<EAxis::Type>, ChildProperty.ToSharedRef())
+							.ExcludedEnumerators({ EAxis::None })
+							.Visibility(VisibilityAttribute)
+							.IsEnabled(EnabledAttribute)
+						]
+						.ExtensionContent()
+						[
+							ChildProperty->CreateDefaultPropertyButtonWidgets()
+						];
+				}
+				else
+				{
+					StructBuilder.AddProperty(ChildProperty.ToSharedRef());
+				}
+			}
+		}
+	}
+}
diff --git a/Engine/Plugins/Animation/SkeletalMeshModelingTools/Source/SkeletalMeshModelingTools/Private/DetailCustomization/SKMMTMirroringPropertiesAxisCustomization.h b/Engine/Plugins/Animation/SkeletalMeshModelingTools/Source/SkeletalMeshModelingTools/Private/DetailCustomization/SKMMTMirroringPropertiesAxisCustomization.h
new file mode 100644
--- /dev/null
+++ b/Engine/Plugins/Animation/SkeletalMeshModelingTools/Source/SkeletalMeshModelingTools/Private/DetailCustomization/SKMMTMirroringPropertiesAxisCustomization.h
@@ -0,0 +1,23 @@
+// Copyright Epic Games, Inc. All Rights Reserved.
+
+#pragma once
+
+#include "IPropertyTypeCustomization.h"
+
+namespace UE::SkeletalMeshModelingTools
+{
+	/** Property type customization for EAxis in the Skeletal Mesh editor */
+	class FSKMMTMirroringPropertiesAxisCustomization
+		: public IPropertyTypeCustomization
+	{
+	public:
+		/** Creates an instance of this property type customization */
+		static TSharedRef<IPropertyTypeCustomization> MakeInstance();
+
+	private:
+		//~ Begin IPropertyTypeCustomization interface
+		virtual void CustomizeHeader(TSharedRef<IPropertyHandle> InStructPropertyHandle, FDetailWidgetRow& HeaderRow, IPropertyTypeCustomizationUtils& StructCustomizationUtils) override;
+		virtual void CustomizeChildren(TSharedRef<IPropertyHandle> InStructPropertyHandle, IDetailChildrenBuilder& StructBuilder, IPropertyTypeCustomizationUtils& StructCustomizationUtils) override;
+		//~ End IPropertyTypeCustomization interface
+	};
+}
diff --git a/Engine/Plugins/Animation/SkeletalMeshModelingTools/Source/SkeletalMeshModelingTools/Private/DetailCustomization/SKMMTOrientPropertiesAxisCustomization.cpp b/Engine/Plugins/Animation/SkeletalMeshModelingTools/Source/SkeletalMeshModelingTools/Private/DetailCustomization/SKMMTOrientPropertiesAxisCustomization.cpp
new file mode 100644
--- /dev/null
+++ b/Engine/Plugins/Animation/SkeletalMeshModelingTools/Source/SkeletalMeshModelingTools/Private/DetailCustomization/SKMMTOrientPropertiesAxisCustomization.cpp
@@ -0,0 +1,99 @@
+// Copyright Epic Games, Inc. All Rights Reserved.
+
+#include "SKMMTOrientPropertiesAxisCustomization.h"
+
+#include "DetailWidgetRow.h"
+#include "IDetailChildrenBuilder.h"
+#include "PropertyHandle.h"
+#include "SkeletalMesh/SkeletonEditingTool.h"
+#include "Widgets/SSKMEnumPropertyComboBox.h"
+
+namespace UE::SkeletalMeshModelingTools
+{
+	TSharedRef<IPropertyTypeCustomization> FSKMMTOrientPropertiesAxisCustomization::MakeInstance()
+	{
+		return MakeShared<FSKMMTOrientPropertiesAxisCustomization>();
+	}
+
+	void FSKMMTOrientPropertiesAxisCustomization::CustomizeHeader(
+		TSharedRef<IPropertyHandle> InStructPropertyHandle,
+		FDetailWidgetRow& HeaderRow,
+		IPropertyTypeCustomizationUtils& StructCustomizationUtils)
+	{
+	}
+
+	void FSKMMTOrientPropertiesAxisCustomization::CustomizeChildren(
+		TSharedRef<IPropertyHandle> InStructPropertyHandle,
+		IDetailChildrenBuilder& StructBuilder,
+		IPropertyTypeCustomizationUtils& StructCustomizationUtils)
+	{
+		TArray<UObject*> OuterObjects;
+		InStructPropertyHandle->GetOuterObjects(OuterObjects);
+
+		const bool bValid = OuterObjects.Num() == 1 && OuterObjects[0] && OuterObjects[0]->GetClass() == UOrientingProperties::StaticClass();
+		if (!ensureMsgf(bValid, TEXT("Unexpected unexpected objects wheen trying to customize the Orient Properties' Orient Axis")))
+		{
+			return;
+		}
+		const TWeakObjectPtr<UOrientingProperties> WeakOrientProperties = CastChecked<UOrientingProperties>(OuterObjects[0]);
+
+		const TAttribute<EVisibility> VisibilityAttribute = TAttribute<EVisibility>::CreateLambda([WeakOrientProperties]()
+			{
+				const bool bVisible =
+					WeakOrientProperties.IsValid() &&
+					WeakOrientProperties->ParentTool.IsValid() &&
+					WeakOrientProperties->ParentTool->GetOperation() != EEditingOperation::Create;
+
+				return bVisible ? EVisibility::Visible : EVisibility::Hidden;
+			});
+
+		const TAttribute<bool> EnabledAttribute = TAttribute<bool>::CreateLambda([WeakOrientProperties]()
+			{
+				return
+					WeakOrientProperties.IsValid() &&
+					WeakOrientProperties->ParentTool.IsValid() &&
+					!WeakOrientProperties->ParentTool->GetSelection().IsEmpty() &&
+					WeakOrientProperties->ParentTool->GetOperation() != EEditingOperation::Parent;
+			});
+
+		uint32 NumChildren = 0;
+		InStructPropertyHandle->GetNumChildren(NumChildren);
+		
+		const TSharedPtr<IPropertyHandle> PrimaryOrientAxisProperty = InStructPropertyHandle->GetChildHandle(GET_MEMBER_NAME_CHECKED(FOrientOptions, Primary));
+		const TSharedPtr<IPropertyHandle> SecondaryOrientAxisProperty = InStructPropertyHandle->GetChildHandle(GET_MEMBER_NAME_CHECKED(FOrientOptions, Secondary));
+
+		for (uint32 ChildIndex = 0; ChildIndex < NumChildren; ChildIndex++)
+		{
+			if (const TSharedPtr<IPropertyHandle> ChildProperty = InStructPropertyHandle->GetChildHandle(ChildIndex))
+			{
+				const FProperty* Property = ChildProperty->GetProperty();
+				const FName PropertyName = Property ? Property->GetFName() : NAME_None;
+
+				if (PropertyName == GET_MEMBER_NAME_CHECKED(FOrientOptions, Primary) ||
+					PropertyName == GET_MEMBER_NAME_CHECKED(FOrientOptions, Secondary))
+				{
+					StructBuilder.AddCustomRow(ChildProperty->GetPropertyDisplayName())
+						.NameContent()
+						[
+							ChildProperty->CreatePropertyNameWidget()
+						]
+						.ValueContent()
+						[
+							SNew(SSKMEnumPropertyComboBox<EOrientAxis>, ChildProperty.ToSharedRef())
+							.ExcludedEnumerators({ EOrientAxis::None })
+							.Visibility(VisibilityAttribute)
+							.IsEnabled(EnabledAttribute)
+						]
+						.ExtensionContent()
+						[
+							ChildProperty->CreateDefaultPropertyButtonWidgets()
+						];
+				}
+				else
+				{
+					StructBuilder.AddProperty(ChildProperty.ToSharedRef());
+				}
+			}
+		}
+	}
+}
diff --git a/Engine/Plugins/Animation/SkeletalMeshModelingTools/Source/SkeletalMeshModelingTools/Private/DetailCustomization/SKMMTOrientPropertiesAxisCustomization.h b/Engine/Plugins/Animation/SkeletalMeshModelingTools/Source/SkeletalMeshModelingTools/Private/DetailCustomization/SKMMTOrientPropertiesAxisCustomization.h
new file mode 100644
--- /dev/null
+++ b/Engine/Plugins/Animation/SkeletalMeshModelingTools/Source/SkeletalMeshModelingTools/Private/DetailCustomization/SKMMTOrientPropertiesAxisCustomization.h
@@ -0,0 +1,23 @@
+// Copyright Epic Games, Inc. All Rights Reserved.
+
+#pragma once
+
+#include "IPropertyTypeCustomization.h"
+
+namespace UE::SkeletalMeshModelingTools
+{
+	/** Property type customization for EAxis in the Skeletal Mesh editor */
+	class FSKMMTOrientPropertiesAxisCustomization
+		: public IPropertyTypeCustomization
+	{
+	public:
+		/** Creates an instance of this property type customization */
+		static TSharedRef<IPropertyTypeCustomization> MakeInstance();
+
+	private:
+		//~ Begin IPropertyTypeCustomization interface
+		virtual void CustomizeHeader(TSharedRef<IPropertyHandle> InStructPropertyHandle, FDetailWidgetRow& HeaderRow, IPropertyTypeCustomizationUtils& StructCustomizationUtils) override;
+		virtual void CustomizeChildren(TSharedRef<IPropertyHandle> InStructPropertyHandle, IDetailChildrenBuilder& StructBuilder, IPropertyTypeCustomizationUtils& StructCustomizationUtils) override;
+		//~ End IPropertyTypeCustomization interface
+	};
+}
diff --git a/Engine/Plugins/Animation/SkeletalMeshModelingTools/Source/SkeletalMeshModelingTools/Private/DetailCustomization/SkeletonEditingToolPropertyCustomizations.cpp b/Engine/Plugins/Animation/SkeletalMeshModelingTools/Source/SkeletalMeshModelingTools/Private/DetailCustomization/SkeletonEditingToolPropertyCustomizations.cpp
--- a/Engine/Plugins/Animation/SkeletalMeshModelingTools/Source/SkeletalMeshModelingTools/Private/DetailCustomization/SkeletonEditingToolPropertyCustomizations.cpp
+++ b/Engine/Plugins/Animation/SkeletalMeshModelingTools/Source/SkeletalMeshModelingTools/Private/DetailCustomization/SkeletonEditingToolPropertyCustomizations.cpp
@@ -5,6 +5,8 @@
 #include "DetailLayoutBuilder.h"
 #include "Templates/UniquePtr.h"
 #include "SkeletalMesh/SkeletonEditingTool.h"
+#include "SKMMTMirroringPropertiesAxisCustomization.h"
+#include "SKMMTOrientPropertiesAxisCustomization.h"
 #include "Widgets/SBoxPanel.h"
 #include "Widgets/Input/SEditableTextBox.h"
 #include "Widgets/Input/SSegmentedControl.h"
@@ -675,7 +677,13 @@ const TArray<FName>& FMirroringPropertiesDetailCustomization::GetProperties() co
 
 void FMirroringPropertiesDetailCustomization::CustomizeDetails(IDetailLayoutBuilder& DetailBuilder)
 {
-	TWeakObjectPtr<UMirroringProperties> Properties = GetCustomizedObject<UMirroringProperties>(DetailBuilder);
+	using namespace UE::SkeletalMeshModelingTools;
+
+	DetailBuilder.RegisterInstancedCustomPropertyTypeLayout(
+		FMirrorOptions::StaticStruct()->GetFName(),
+		FOnGetPropertyTypeCustomizationInstance::CreateStatic(&FSKMMTMirroringPropertiesAxisCustomization::MakeInstance));
+
+	const TWeakObjectPtr<UMirroringProperties> Properties = GetCustomizedObject<UMirroringProperties>(DetailBuilder);
 	if (!GetParentTool(Properties).IsValid())
 	{
 		return;
@@ -687,7 +695,7 @@ void FMirroringPropertiesDetailCustomization::CustomizeDetails(IDetailLayoutBuil
 	static const FName MirrorCategoryName("Mirror");
 	IDetailCategoryBuilder& MirrorCategory = DetailBuilder.EditCategory(MirrorCategoryName, FText::GetEmpty(), ECategoryPriority::Important);
 	
-	MirrorCategory.AddCustomRow(LOCTEXT("MirrorButtonRow", "Mirror"), false)
+	MirrorCategory.AddCustomRow(LOCTEXT("MirrorButtonRow", "Mirror"))
 	[
 		SNew(SHorizontalBox)
 		+SHorizontalBox::Slot()
@@ -725,6 +733,12 @@ const TArray<FName>& FOrientingPropertiesDetailCustomization::GetProperties() co
 
 void FOrientingPropertiesDetailCustomization::CustomizeDetails(IDetailLayoutBuilder& DetailBuilder)
 {
+	using namespace UE::SkeletalMeshModelingTools;
+
+	DetailBuilder.RegisterInstancedCustomPropertyTypeLayout(
+		FOrientOptions::StaticStruct()->GetFName(),
+		FOnGetPropertyTypeCustomizationInstance::CreateStatic(&FSKMMTOrientPropertiesAxisCustomization::MakeInstance));
+
 	TWeakObjectPtr<UOrientingProperties> Properties = GetCustomizedObject<UOrientingProperties>(DetailBuilder);
 	if (!GetParentTool(Properties).IsValid())
 	{
@@ -734,8 +748,8 @@ void FOrientingPropertiesDetailCustomization::CustomizeDetails(IDetailLayoutBuil
 	
 	UpdateProperties(DetailBuilder, EditionVisibility, IsEnabledBySelection());
 	static const FName OrientCategoryName("Orient");
-	IDetailCategoryBuilder& MirrorCategory = DetailBuilder.EditCategory(OrientCategoryName, FText::GetEmpty(), ECategoryPriority::Important);
-	MirrorCategory.AddCustomRow(LOCTEXT("OrientButtonRow", "Orient"), false)
+	IDetailCategoryBuilder& OrientCategory = DetailBuilder.EditCategory(OrientCategoryName, FText::GetEmpty(), ECategoryPriority::Important);
+	OrientCategory.AddCustomRow(LOCTEXT("OrientButtonRow", "Orient"), false)
 	[
 		SNew(SHorizontalBox)
 		+SHorizontalBox::Slot()
diff --git a/Engine/Plugins/Animation/SkeletalMeshModelingTools/Source/SkeletalMeshModelingTools/Private/Widgets/SSKMEnumPropertyComboBox.h b/Engine/Plugins/Animation/SkeletalMeshModelingTools/Source/SkeletalMeshModelingTools/Private/Widgets/SSKMEnumPropertyComboBox.h
new file mode 100644
--- /dev/null
+++ b/Engine/Plugins/Animation/SkeletalMeshModelingTools/Source/SkeletalMeshModelingTools/Private/Widgets/SSKMEnumPropertyComboBox.h
@@ -0,0 +1,181 @@
+// Copyright Epic Games, Inc. All Rights Reserved.
+
+#pragma once
+
+#include "PropertyHandle.h"
+#include "SEnumCombo.h"
+#include "Widgets/SCompoundWidget.h"
+
+class IPropertyHandle;
+
+namespace UE::SkeletalMeshModelingTools
+{
+	/** A combo box for enum values. Useful to exclude elements */
+	template <typename TEnumType>
+	class SSKMEnumPropertyComboBox
+		: public SCompoundWidget
+	{
+	public:
+		SLATE_BEGIN_ARGS(SSKMEnumPropertyComboBox<TEnumType>)
+			{}
+
+			/** The visibility attribute for this widget */
+			SLATE_ATTRIBUTE(EVisibility, Visibility);
+
+			/** The enabled attribute for this widget. Only considered when the enum property handle is editable */
+			SLATE_ATTRIBUTE(bool, IsEnabled);
+
+			/** Enumerators that should be excluded from the combo box */
+			SLATE_ARGUMENT(TSet<TEnumType>, ExcludedEnumerators);
+		
+		SLATE_END_ARGS()
+
+		/** Constructs this widget */
+		void Construct(const FArguments& InArgs, const TSharedRef<IPropertyHandle>& InEnumPropertyHandle)
+		{
+			const UEnum* Enum = StaticEnum<TEnumType>();
+			if (!ensureMsgf(Enum, TEXT("Invalid enum type when trying to show an SSKMEnumComboBox in the SKM editor. Cannot display enum.")))
+			{
+				return;
+			}			
+
+			UnderlyingType = [&InEnumPropertyHandle, this]() -> const FFieldClass*
+				{
+					if (const FEnumProperty* EnumProperty = CastField<FEnumProperty>(InEnumPropertyHandle->GetProperty()))
+					{
+						const FNumericProperty* UnderlyingTypeProperty = EnumProperty ? EnumProperty->GetUnderlyingProperty() : nullptr;
+						return UnderlyingTypeProperty ? UnderlyingTypeProperty->GetClass() : nullptr;
+					}
+					else if (const FByteProperty* ByteProperty = CastField<FByteProperty>(InEnumPropertyHandle->GetProperty()))
+					{
+						return ByteProperty->GetClass();
+					}
+					return nullptr;
+				}();
+			if (!ensureMsgf(UnderlyingType, TEXT("Unsupported enum type when trying to show an SSKMEnumComboBox in the SKM editor. Cannot display enum.")))
+			{
+				return;
+			}
+
+			EnumPropertyHandle = InEnumPropertyHandle;
+			IsEnabledAttribute = InArgs._IsEnabled;
+
+			TArray<int32> AvaialableValues;
+			for (int32 EnumIndex = 0; EnumIndex < Enum->NumEnums() - 1; EnumIndex++)
+			{
+				const int32 EnumValue = Enum->GetValueByIndex(EnumIndex);
+				if (!InArgs._ExcludedEnumerators.Contains(static_cast<TEnumType>(EnumValue)))
+				{
+					const int32 Value = static_cast<int32>(Enum->GetValueByIndex(EnumIndex));
+					AvaialableValues.Add(Value);
+				}
+			}
+			
+			if (AvaialableValues.IsEmpty())
+			{
+				return;
+			}
+
+			if (!AvaialableValues.Contains(GetValue()))
+			{
+				// Set X for the value if none is hidden and selected
+				SetValue(AvaialableValues[0], ESelectInfo::Direct);
+			}
+
+			ChildSlot
+			[
+				SNew(SEnumComboBox, StaticEnum<TEnumType>())
+				.Font(FAppStyle::GetFontStyle("PropertyWindow.NormalFont"))		
+				.TextStyle(FAppStyle::Get(), TEXT("SmallText"))
+				.Visibility(InArgs._Visibility)
+				.IsEnabled(this, &SSKMEnumPropertyComboBox::IsEnabled)
+				.CurrentValue(this, &SSKMEnumPropertyComboBox::GetValue)
+				.OnEnumSelectionChanged(this, &SSKMEnumPropertyComboBox::SetValue)
+				.EnumValueSubset(AvaialableValues)
+			];
+		}
+
+	private:
+		/** Returns true if the widget is enabled */
+		bool IsEnabled() const
+		{
+			if (!EnumPropertyHandle.IsValid() ||
+				!EnumPropertyHandle->IsEditable())
+			{
+				return false;
+			}
+
+			if (IsEnabledAttribute.IsSet())
+			{
+				return IsEnabledAttribute.Get();
+			}
+
+			return true;
+		}
+
+		/** Returns the currently selected axis */
+		int32 GetValue() const
+		{
+			if (EnumPropertyHandle.IsValid() &&
+				EnumPropertyHandle->IsValidHandle())
+			{
+				if (UnderlyingType == FByteProperty::StaticClass())
+				{
+					uint8 Value;
+					if (EnumPropertyHandle->GetValue(Value) == FPropertyAccess::Success)
+					{
+						return Value;
+					}
+				}
+				else if (UnderlyingType == FIntProperty::StaticClass())
+				{
+					int32 Value;
+					if (EnumPropertyHandle->GetValue(Value) == FPropertyAccess::Success)
+					{
+						return Value;
+					}
+				}
+				else
+				{
+					ensureMsgf(0, TEXT("Unsupported type"));
+				}
+			}
+
+			return 0;
+		}
+
+		/** Called when an axis was selected */
+		void SetValue(int32 EnumValue, ESelectInfo::Type SelectInfo)
+		{
+			if (EnumPropertyHandle.IsValid() &&
+				EnumPropertyHandle->IsValidHandle())
+			{
+				EnumPropertyHandle->NotifyPreChange();
+
+				if (UnderlyingType == FByteProperty::StaticClass())
+				{
+					EnumPropertyHandle->SetValue(static_cast<uint8>(EnumValue));
+				}
+				else if (UnderlyingType == FIntProperty::StaticClass())
+				{
+					EnumPropertyHandle->SetValue(EnumValue);
+				}
+				else
+				{
+					ensureMsgf(0, TEXT("Unsupported type"));
+				}
+
+				EnumPropertyHandle->NotifyPostChange(EPropertyChangeType::ValueSet);
+			}
+		}
+
+		/** The size type of the enum */
+		const FFieldClass* UnderlyingType = nullptr;
+
+		/** The enabled attribute for this widget. Only considered when the enum property handle is editable. */
+		TAttribute<bool> IsEnabledAttribute;
+
+		/** The property handle to the enum */
+		TSharedPtr<IPropertyHandle> EnumPropertyHandle;
+	};
+}
