[SkeletalMeshModelingTool] added left click marquee selection, requires enabling 5.8 gizmo

#Synced-CL 49873372

diff --git a/Engine/Plugins/Animation/SkeletalMeshModelingTools/Source/SkeletalMeshModelingTools/Private/Selection/SKMGeometrySelectionTransformTweaker.cpp b/Engine/Plugins/Animation/SkeletalMeshModelingTools/Source/SkeletalMeshModelingTools/Private/Selection/SKMGeometrySelectionTransformTweaker.cpp
new file mode 100644
--- /dev/null
+++ b/Engine/Plugins/Animation/SkeletalMeshModelingTools/Source/SkeletalMeshModelingTools/Private/Selection/SKMGeometrySelectionTransformTweaker.cpp
@@ -0,0 +1,143 @@
+// Copyright Epic Games, Inc. All Rights Reserved.
+
+#include "Selection/SKMGeometrySelectionTransformTweaker.h"
+
+#include "InteractiveToolsContext.h"
+#include "ModelingSelectionInteraction.h"
+#include "SkeletalMeshGizmoUtils.h"
+#include "BaseGizmos/TransformProxy.h"
+#include "EditorGizmos/EditorTransformGizmoUtil.h"
+#include "EditorGizmos/TransformGizmo.h"
+#include "Selection/GeometrySelectionManager.h"
+
+#define LOCTEXT_NAMESPACE "USkeletalMeshGeometrySelectionTransformManipulator"
+
+using namespace UE::Geometry;
+
+void USkeletalMeshGeometrySelectionTransformTweaker::Setup(UGeometrySelectionManager* InSelectionManager)
+{
+	SelectionManager = InSelectionManager;
+
+	SelectionManager->OnSelectionModified.AddUObject(this, &USkeletalMeshGeometrySelectionTransformTweaker::OnSelectionManager_SelectionModified);
+}
+
+void USkeletalMeshGeometrySelectionTransformTweaker::Shutdown()
+{
+	if (SelectionManager.IsValid())
+	{
+		SelectionManager->OnSelectionModified.RemoveAll(this);
+	}
+}
+
+const FTransform& USkeletalMeshGeometrySelectionTransformTweaker::GetSelectionFrameTransform()
+{
+	return CachedSelectionFrameTransform;
+}
+
+void USkeletalMeshGeometrySelectionTransformTweaker::BeginTransformEdit()
+{
+	StartSelectionFrameTransform = CachedSelectionFrameTransform;
+	Scale = FVector::One();
+	// Sometimes when moving the 5.8 gizmo using indirect manipulation (ctrl left click)
+	// BeginTransformEdit can be called twice in the same frame, coming from
+	// UViewportClickDragBehavior::BeginCapture & UpdateCapture. To be investigated
+	if (!SelectionManager->IsInActiveTransformation())
+	{
+		SelectionManager->BeginTransformation();
+	}
+}
+
+void USkeletalMeshGeometrySelectionTransformTweaker::EndTransformEdit()
+{
+	SelectionManager->EndTransformation();
+	
+	UpdateSelectionFrameTransform();
+}
+
+bool USkeletalMeshGeometrySelectionTransformTweaker::IsEditingTransform() const
+{
+	return SelectionManager->IsInActiveTransformation();
+}
+
+void USkeletalMeshGeometrySelectionTransformTweaker::TweakTransform(FVector& InDrag, FRotator& InRot, FVector& InScale, bool bInIsWorldSpace)
+{
+	CachedSelectionFrameTransform.AddToTranslation(InDrag);
+	CachedSelectionFrameTransform.SetRotation(InRot.Quaternion() * CachedSelectionFrameTransform.GetRotation());
+	Scale += InScale;
+
+	if (bInIsWorldSpace)
+	{
+		// For a world frame gizmo, the scaling needs to happen in world aligned gizmo space, but the 
+		// rotation is still encoded in the local gizmo frame change.
+		FQuat RotationToApply = CachedSelectionFrameTransform.GetRotation() * StartSelectionFrameTransform.GetRotation().Inverse();
+		SelectionManager->UpdateTransformation([this, RotationToApply](int32 VertexID, const FVector3d& InitialPosition, const FTransform3d& WorldTransform)
+			{
+				FVector3d PosLocal = InitialPosition;
+				FVector3d PosWorld = WorldTransform.TransformPosition(PosLocal);
+				FVector3d PosWorldGizmo = PosWorld - StartSelectionFrameTransform.GetLocation();
+
+				FVector3d NewPosWorld = RotationToApply * (Scale * PosWorldGizmo) + CachedSelectionFrameTransform.GetLocation();
+				FVector3d NewPosLocal = WorldTransform.InverseTransformPosition(NewPosWorld);
+				return NewPosLocal;
+			});	
+	}
+	else
+	{
+		SelectionManager->UpdateTransformation([this](int32 VertexID, const FVector3d& InitialPosition, const FTransform3d& WorldTransform)
+			{
+				// For a local gizmo, we just get the coordinates in the original frame, scale in that frame,
+				// then interpret them as coordinates in the new frame.
+				FVector3d PosLocal = InitialPosition;
+				FVector3d PosWorld = WorldTransform.TransformPosition(PosLocal);
+				FVector3d PosGizmo = StartSelectionFrameTransform.InverseTransformPosition(PosWorld);
+				PosGizmo = Scale * PosGizmo;
+				FVector3d NewPosWorld = CachedSelectionFrameTransform.TransformPosition(PosGizmo);
+				FVector3d NewPosLocal = WorldTransform.InverseTransformPosition(NewPosWorld);
+				return NewPosLocal;
+			});	
+	}
+	
+}
+
+void USkeletalMeshGeometrySelectionTransformTweaker::SetLocalFrameMode(EModelingSelectionInteraction_LocalFrameMode InLocalFrameMode)
+{
+	if (LocalFrameMode != InLocalFrameMode)
+	{
+		LocalFrameMode = InLocalFrameMode;
+		UpdateSelectionFrameTransform();
+	}
+}
+
+EModelingSelectionInteraction_LocalFrameMode USkeletalMeshGeometrySelectionTransformTweaker::GetLocalFrameMode() const
+{
+	return LocalFrameMode;
+}
+
+void USkeletalMeshGeometrySelectionTransformTweaker::OnSelectionManager_SelectionModified()
+{
+	UpdateSelectionFrameTransform();	
+}
+
+void USkeletalMeshGeometrySelectionTransformTweaker::UpdateSelectionFrameTransform()
+{
+	if (!SelectionManager->HasSelection())
+	{
+		return;
+	}
+	
+	FFrame3d SelectionFrame;
+
+	if (LocalFrameMode == EModelingSelectionInteraction_LocalFrameMode::FromGeometry)
+	{
+		SelectionManager->GetSelectionWorldFrame(SelectionFrame);
+	}
+	else
+	{
+		SelectionManager->GetTargetWorldFrame(SelectionFrame);
+	}
+	
+	CachedSelectionFrameTransform = SelectionFrame.ToFTransform();
+}
+
+#undef LOCTEXT_NAMESPACE 
+
diff --git a/Engine/Plugins/Animation/SkeletalMeshModelingTools/Source/SkeletalMeshModelingTools/Private/Selection/SKMGeometrySelectionTransformTweaker.h b/Engine/Plugins/Animation/SkeletalMeshModelingTools/Source/SkeletalMeshModelingTools/Private/Selection/SKMGeometrySelectionTransformTweaker.h
new file mode 100644
--- /dev/null
+++ b/Engine/Plugins/Animation/SkeletalMeshModelingTools/Source/SkeletalMeshModelingTools/Private/Selection/SKMGeometrySelectionTransformTweaker.h
@@ -0,0 +1,48 @@
+// Copyright Epic Games, Inc. All Rights Reserved.
+
+#pragma once
+
+#include "CoreMinimal.h"
+#include "ModelingSelectionInteraction.h"
+
+#include "SKMGeometrySelectionTransformTweaker.generated.h"
+
+class UGeometrySelectionManager;
+class UTransformProxy;
+class UTransformGizmo;
+
+
+UCLASS(MinimalAPI)
+class USkeletalMeshGeometrySelectionTransformTweaker : public UObject
+{
+	GENERATED_BODY()
+	
+public:
+	void Setup(UGeometrySelectionManager* InSelectionManager);
+	void Shutdown();
+
+	const FTransform& GetSelectionFrameTransform();
+
+	void BeginTransformEdit();
+	void EndTransformEdit();
+	bool IsEditingTransform() const;
+	void TweakTransform(FVector& InDrag, FRotator& InRot, FVector& InScale, bool bInIsWorldSpace);
+
+	void SetLocalFrameMode(EModelingSelectionInteraction_LocalFrameMode InLocalFrameMode);
+	EModelingSelectionInteraction_LocalFrameMode GetLocalFrameMode() const;
+	
+	void UpdateSelectionFrameTransform();
+protected:
+
+	void OnSelectionManager_SelectionModified();
+	
+	TWeakObjectPtr<UGeometrySelectionManager> SelectionManager;
+
+	FTransform CachedSelectionFrameTransform;
+	
+	FTransform StartSelectionFrameTransform;
+
+	FVector Scale;
+
+	EModelingSelectionInteraction_LocalFrameMode LocalFrameMode = EModelingSelectionInteraction_LocalFrameMode::FromGeometry;
+};
\ No newline at end of file
diff --git a/Engine/Plugins/Animation/SkeletalMeshModelingTools/Source/SkeletalMeshModelingTools/Private/SkeletalMeshEditingCache.cpp b/Engine/Plugins/Animation/SkeletalMeshModelingTools/Source/SkeletalMeshModelingTools/Private/SkeletalMeshEditingCache.cpp
--- a/Engine/Plugins/Animation/SkeletalMeshModelingTools/Source/SkeletalMeshModelingTools/Private/SkeletalMeshEditingCache.cpp
+++ b/Engine/Plugins/Animation/SkeletalMeshModelingTools/Source/SkeletalMeshModelingTools/Private/SkeletalMeshEditingCache.cpp
@@ -53,6 +53,8 @@ void FSkeletalMeshEditingCacheNotifier::HandleNotification(const TArray<FName>&
 
 void USkeletalMeshEditingCache::Spawn(UWorld* World, USkeletalMeshComponent* InSkeletalMeshComponent, EMeshLODIdentifier InLOD, const FDelegates& InDelegates)
 {
+	SetFlags(RF_Transactional);
+	
 	Delegates = InDelegates;
 	SkeletalMeshComponent = InSkeletalMeshComponent;
 	
@@ -135,19 +137,19 @@ void USkeletalMeshEditingCache::Spawn(UWorld* World, USkeletalMeshComponent* InS
 			},
 		GetEditingMeshComponent()->GetChangeCount());
 
-	DynamicMeshStateWatcher.Initialize(
+	DynamicMeshStateUpdater.Initialize(
 		[this]()
 			{
-				FDynamicMeshState State;
-				State.ChangeCount = GetEditingMeshComponent()->GetChangeCount();
-				State.bForceEnable = bForceEnableDynamicMesh;
-				return State;
+				FDynamicMeshStateFactors Factors;
+				Factors.ChangeCount = GetEditingMeshComponent()->GetChangeCount();
+				Factors.bForceEnable = bForceEnableDynamicMesh;
+				return Factors;
 			},
-		[this](FDynamicMeshState State)
+		[this](FDynamicMeshStateFactors Factors)
 			{
-				bIsDynamicMeshEnabled = State.bForceEnable || State.ChangeCount > 0;
+				bIsDynamicMeshEnabled = Factors.bForceEnable || Factors.ChangeCount > 0;
 			},
-		FDynamicMeshState(GetEditingMeshComponent()->GetChangeCount(), bForceEnableDynamicMesh));
+		FDynamicMeshStateFactors(GetEditingMeshComponent()->GetChangeCount(), bForceEnableDynamicMesh));
 	
 	SkeletonChangeCountWatcher.Initialize(
 		[this]()
@@ -180,65 +182,65 @@ void USkeletalMeshEditingCache::Spawn(UWorld* World, USkeletalMeshComponent* InS
 			},
 		GetEditingMeshComponent()->GetSkeletonChangeTracker().GetChangeCount());
 
-	MeshVisibilityWatcher.Initialize(
+	MeshVisibilityUpdater.Initialize(
 		[this]()
 			{
-				FMeshVisibilityState State;
-				State.bIsDynamicMeshEnabled = bIsDynamicMeshEnabled;
-				State.bCacheVisibility = bCacheVisibility;
-				return State;
+				FMeshVisibilityFactors Factors;
+				Factors.bIsDynamicMeshEnabled = bIsDynamicMeshEnabled;
+				Factors.bCacheVisibility = bCacheVisibility;
+				return Factors;
 			},
-		[this](FMeshVisibilityState State)
+		[this](FMeshVisibilityFactors Factors)
 			{
 				PreviewMesh->SetVisible(
-					State.bIsDynamicMeshEnabled && State.bCacheVisibility);
+					Factors.bIsDynamicMeshEnabled && Factors.bCacheVisibility);
 
 				if (SkeletalMeshComponent.IsValid())
 				{
 					SkeletalMeshComponent->SetVisibility(
-						!State.bIsDynamicMeshEnabled && State.bCacheVisibility);
+						!Factors.bIsDynamicMeshEnabled && Factors.bCacheVisibility);
 				}
 			},
-		FMeshVisibilityState(bIsDynamicMeshEnabled, bCacheVisibility));
+		FMeshVisibilityFactors(bIsDynamicMeshEnabled, bCacheVisibility));
 
-	SkeletonVisibilityWatcher.Initialize(
+	SkeletonVisibilityUpdater.Initialize(
 		[this]()
 			{
-				FSkeletonVisibilityState State;
-				State.bIsDynamicMeshSkeletonEnabled = bIsDynamicMeshSkeletonEnabled;
-				State.CacheSkeletonDrawMode = CacheSkeletonDrawMode;
-				return State;
+				FSkeletonVisibilityFactors Factors;
+				Factors.bIsDynamicMeshSkeletonEnabled = bIsDynamicMeshSkeletonEnabled;
+				Factors.CacheSkeletonDrawMode = CacheSkeletonDrawMode;
+				return Factors;
 			},
-		[this](FSkeletonVisibilityState State)
+		[this](FSkeletonVisibilityFactors Factors)
 			{
-				if (State.bIsDynamicMeshSkeletonEnabled)
+				if (Factors.bIsDynamicMeshSkeletonEnabled)
 				{
-					DynamicMeshSkeletonDrawMode = State.CacheSkeletonDrawMode;
+					DynamicMeshSkeletonDrawMode = Factors.CacheSkeletonDrawMode;
 					SetSkeletalMeshSkeletonDrawMode(ESkeletonDrawMode::Hidden);
 				}
 				else
 				{
 					DynamicMeshSkeletonDrawMode = ESkeletonDrawMode::Hidden;
-					SetSkeletalMeshSkeletonDrawMode(State.CacheSkeletonDrawMode);
+					SetSkeletalMeshSkeletonDrawMode(Factors.CacheSkeletonDrawMode);
 				}
 			},
-		FSkeletonVisibilityState(bIsDynamicMeshSkeletonEnabled, CacheSkeletonDrawMode)
+		FSkeletonVisibilityFactors(bIsDynamicMeshSkeletonEnabled, CacheSkeletonDrawMode)
 	);
 
-	BoneManipulationWatcher.Initialize(
+	BoneManipulationUpdater.Initialize(
 		[this]()
 			{
-				FBoneManipulationState State;
-				State.bIsDynamicMeshSkeletonEnabled = bIsDynamicMeshSkeletonEnabled;
-				State.bCacheBoneManipulation = bCacheBoneManipulation;
-				return State;
+				FBoneManipulationFactors Factors;
+				Factors.bIsDynamicMeshSkeletonEnabled = bIsDynamicMeshSkeletonEnabled;
+				Factors.bCacheBoneManipulation = bCacheBoneManipulation;
+				return Factors;
 			},
-		[this](FBoneManipulationState State)
+		[this](FBoneManipulationFactors Factors)
 			{
 				Delegates.ToggleSkeletalMeshBoneManipulationDelegate.ExecuteIfBound(
-					!State.bIsDynamicMeshSkeletonEnabled && bCacheBoneManipulation);
+					!Factors.bIsDynamicMeshSkeletonEnabled && bCacheBoneManipulation);
 			},
-		FBoneManipulationState(bIsDynamicMeshSkeletonEnabled, bCacheBoneManipulation)
+		FBoneManipulationFactors(bIsDynamicMeshSkeletonEnabled, bCacheBoneManipulation)
 	);
 
 	PreviewMeshVisibilityWatcher.Initialize(
@@ -311,11 +313,11 @@ void USkeletalMeshEditingCache::DiscardChanges()
 void USkeletalMeshEditingCache::Tick()
 {
 	ChangeCountWatcher.CheckAndUpdate();
-	DynamicMeshStateWatcher.CheckAndUpdate();
+	DynamicMeshStateUpdater.CheckAndUpdate();
 	SkeletonChangeCountWatcher.CheckAndUpdate();
-	MeshVisibilityWatcher.CheckAndUpdate();
-	SkeletonVisibilityWatcher.CheckAndUpdate();
-	BoneManipulationWatcher.CheckAndUpdate();
+	MeshVisibilityUpdater.CheckAndUpdate();
+	SkeletonVisibilityUpdater.CheckAndUpdate();
+	BoneManipulationUpdater.CheckAndUpdate();
 	PreviewMeshVisibilityWatcher.CheckAndUpdate();
 	
 	if (bIsDynamicMeshEnabled)
@@ -511,22 +513,15 @@ FTransform USkeletalMeshEditingCache::GetTransform()
 	return PreviewMesh->GetTransform();
 }
 
-void USkeletalMeshEditingCache::HideSkeleton()
-{
-	SavedSkeletonDrawMode = CacheSkeletonDrawMode;
-	CacheSkeletonDrawMode = ESkeletonDrawMode::Hidden;
-}
-
-void USkeletalMeshEditingCache::ShowSkeleton()
-{
-	CacheSkeletonDrawMode = SavedSkeletonDrawMode;
-}
-
 void USkeletalMeshEditingCache::ToggleBoneManipulation(bool bEnable)
 {
 	bCacheBoneManipulation = bEnable;
 }
 
+void USkeletalMeshEditingCache::SetSkeletonDrawMode(ESkeletonDrawMode InSkeletonDrawMode)
+{
+	CacheSkeletonDrawMode = InSkeletonDrawMode;
+}
 
 bool USkeletalMeshEditingCache::IsDynamicMeshSkeletonEnabled() const
 {
@@ -633,11 +628,7 @@ void USkeletalMeshEditingCache::HandleSetMorphTargetWeight(FName MorphTarget, fl
 		return;
 	}
 
-	if (GUndo)
-	{
-		SetFlags(RF_Transactional);
-		Modify();
-	}
+	Modify();
 	
 	MorphTargetOverrides.Emplace(MorphTarget, Weight);
 	
@@ -664,11 +655,7 @@ bool USkeletalMeshEditingCache::GetMorphTargetAutoFill(FName Name)
 
 void USkeletalMeshEditingCache::HandleSetMorphTargetAutoFill(FName Name, bool bAutoFill, float PreviousOverrideWeight)
 {
-	if (GUndo)
-	{
-		SetFlags(RF_Transactional);
-		Modify();
-	}
+	Modify();
 	
 	if (bAutoFill)
 	{
diff --git a/Engine/Plugins/Animation/SkeletalMeshModelingTools/Source/SkeletalMeshModelingTools/Private/SkeletalMeshEditingCache.h b/Engine/Plugins/Animation/SkeletalMeshModelingTools/Source/SkeletalMeshModelingTools/Private/SkeletalMeshEditingCache.h
--- a/Engine/Plugins/Animation/SkeletalMeshModelingTools/Source/SkeletalMeshModelingTools/Private/SkeletalMeshEditingCache.h
+++ b/Engine/Plugins/Animation/SkeletalMeshModelingTools/Source/SkeletalMeshModelingTools/Private/SkeletalMeshEditingCache.h
@@ -76,9 +76,9 @@ public:
 	const TArray<FTransform>& GetComponentSpaceBoneTransforms() const;
 	FTransform GetTransform();
 
-	void HideSkeleton();
-	void ShowSkeleton();
 	void ToggleBoneManipulation(bool bEnable);
+	void SetSkeletonDrawMode(ESkeletonDrawMode InSkeletonDrawMode);
+	
 	bool IsDynamicMeshSkeletonEnabled() const;
 	bool IsDynamicMeshBoneManipulationEnabled() const;
 	int32 GetFirstSelectedBoneIndex() const;
@@ -130,48 +130,48 @@ protected:
 
 	bool bForceEnableDynamicMesh = false;
 	
-	struct FDynamicMeshState
+	struct FDynamicMeshStateFactors
 	{
 		int32 ChangeCount = 0;
 		bool bForceEnable = false;
-		friend bool operator==(const FDynamicMeshState&, const FDynamicMeshState&) = default;	
+		friend bool operator==(const FDynamicMeshStateFactors&, const FDynamicMeshStateFactors&) = default;	
 	};
 	
-	TValueWatcher<FDynamicMeshState> DynamicMeshStateWatcher;
+	TValueWatcher<FDynamicMeshStateFactors> DynamicMeshStateUpdater;
 	bool bIsDynamicMeshEnabled = false;
 
 
 	TValueWatcher<int32> SkeletonChangeCountWatcher;
 	bool bIsDynamicMeshSkeletonEnabled = false;
 
-	struct FMeshVisibilityState
+	struct FMeshVisibilityFactors
 	{
 		bool bIsDynamicMeshEnabled= false;
 		bool bCacheVisibility = true;
-		friend bool operator==(const FMeshVisibilityState&, const FMeshVisibilityState&) = default;
+		friend bool operator==(const FMeshVisibilityFactors&, const FMeshVisibilityFactors&) = default;
 	};
 
-	TValueWatcher<FMeshVisibilityState> MeshVisibilityWatcher;
+	TValueWatcher<FMeshVisibilityFactors> MeshVisibilityUpdater;
 
 
-	struct FSkeletonVisibilityState
+	struct FSkeletonVisibilityFactors
 	{
 		bool bIsDynamicMeshSkeletonEnabled = false;
 		ESkeletonDrawMode CacheSkeletonDrawMode = ESkeletonDrawMode::Default;
-		friend bool operator==(const FSkeletonVisibilityState&, const FSkeletonVisibilityState&) = default;
+		friend bool operator==(const FSkeletonVisibilityFactors&, const FSkeletonVisibilityFactors&) = default;
 	};
 	
-	TValueWatcher<FSkeletonVisibilityState> SkeletonVisibilityWatcher;
+	TValueWatcher<FSkeletonVisibilityFactors> SkeletonVisibilityUpdater;
 
 
-	struct FBoneManipulationState
+	struct FBoneManipulationFactors
 	{
 		bool bIsDynamicMeshSkeletonEnabled = false;
 		bool bCacheBoneManipulation = true;
-		friend bool operator==(const FBoneManipulationState&, const FBoneManipulationState&) = default;
+		friend bool operator==(const FBoneManipulationFactors&, const FBoneManipulationFactors&) = default;
 	};
 	
-	TValueWatcher<FBoneManipulationState> BoneManipulationWatcher;
+	TValueWatcher<FBoneManipulationFactors> BoneManipulationUpdater;
 	
 	TValueWatcher<bool> PreviewMeshVisibilityWatcher;
 
@@ -206,8 +206,6 @@ protected:
 	void SetSkeletalMeshSkeletonDrawMode(ESkeletonDrawMode DrawMode) const;
 	ESkeletonDrawMode GetCurrentSkeletonDrawMode() const;
 	
-	ESkeletonDrawMode SavedSkeletonDrawMode = ESkeletonDrawMode::Default;
-
 	FLinearColor GetDefaultBoneColor(int32 BoneIndex) const;
 
 	UDebugSkelMeshComponent* GetDebugSkelMeshComponent() const;
diff --git a/Engine/Plugins/Animation/SkeletalMeshModelingTools/Source/SkeletalMeshModelingTools/Private/SkeletalMeshModelingToolsEditorMode.cpp b/Engine/Plugins/Animation/SkeletalMeshModelingTools/Source/SkeletalMeshModelingTools/Private/SkeletalMeshModelingToolsEditorMode.cpp
--- a/Engine/Plugins/Animation/SkeletalMeshModelingTools/Source/SkeletalMeshModelingTools/Private/SkeletalMeshModelingToolsEditorMode.cpp
+++ b/Engine/Plugins/Animation/SkeletalMeshModelingTools/Source/SkeletalMeshModelingTools/Private/SkeletalMeshModelingToolsEditorMode.cpp
@@ -18,6 +18,14 @@
 // We are leaving the code in this cpp file, defined out, so that it is easier to bring back if/when the stylus plugin is improved.
 
 #include "SkeletalMeshModelingToolsModule.h"
+#include "Selection/SKMGeometrySelectionTransformTweaker.h"
+#include "ViewportInteractions/SKMModelingToolsClickSelection.h"
+#include "ViewportInteractions/SKMModelingToolsGeometryClickSelection.h"
+#include "ViewportInteractions/SKMModelingToolsGeometryFrustumSelection.h"
+#include "ViewportInteractions/SKMModelingToolsGeometryHoverInteraction.h"
+#include "ViewportInteractions/ViewportInteractionsBehaviorSource.h"
+#include "ViewportInteractions/ViewportInteractionsUtilities.h"
+#include "ViewportInteractions/ViewportMoveYawInteraction.h"
 #include "Widgets/Notifications/SNotificationList.h"
 
 #include UE_INLINE_GENERATED_CPP_BY_NAME(SkeletalMeshModelingToolsEditorMode)
@@ -31,13 +39,13 @@
 #include "AnimationEditorViewportClient.h"
 #include "ToolTargets/SkeletalMeshComponentToolTarget.h"
 #include "Components/SkeletalMeshComponent.h"
+#include "Engine/SkeletalMesh.h"
 
 #include "ConvertToPolygonsTool.h"
 #include "DeformMeshPolygonsTool.h"
 #include "DisplaceMeshTool.h"
 #include "DynamicMeshSculptTool.h"
 #include "EditMeshPolygonsTool.h"
-#include "EditorInteractiveGizmoManager.h"
 #include "HoleFillTool.h"
 #include "IPersonaPreviewScene.h"
 #include "IPersonaToolkit.h"
@@ -74,11 +82,8 @@
 #include "ToolTargets/SkeletalMeshToolTarget.h"
 #include "ToolTargets/SKMBackedDynaMeshComponentToolTarget.h"
 
-#include "ModelingSelectionInteraction.h"
 #include "ModelingToolsEditorModeSettings.h"
-#include "Selection/DynamicMeshSelector.h"
 #include "Selection/GeometrySelectionManager.h"
-#include "UnrealWidget.h"
 #include "ModelingToolsEditorModeStyle.h"
 #include "ToolMenus.h"
 #include "SkeletalMeshToolMenuContext.h"
@@ -380,6 +385,35 @@ void USkeletalMeshModelingToolsEditorMode::Enter()
 
 	InitializeGeometrySelectionSystems();
 	PopulateViewportToolbar();
+
+	if (UModeManagerInteractiveToolsContext* const ModeManagerInteractiveToolsContext = GetModeManager()->GetInteractiveToolsContext())
+	{
+		ModeManagerInteractiveToolsContext->SetViewportInteractionsEnabled(true);
+
+		BuildViewportInteractionsDelegateHandle = ModeManagerInteractiveToolsContext->OnBuildViewportInteractions().AddLambda([WeakThis = TWeakObjectPtr<USkeletalMeshModelingToolsEditorMode>(this)](UViewportInteractionsBehaviorSource* Source)
+		{
+			UE::Editor::ViewportInteractions::AddDefaultCameraMovementInteractions(Source);
+			Source->AddInteraction<USkeletalMeshModelingToolsClickSelection>();
+				
+			check(WeakThis.IsValid());
+			USkeletalMeshModelingToolsGeometryFrustumSelection* GeometryFrustumSelection = Source->AddInteraction<
+				USkeletalMeshModelingToolsGeometryFrustumSelection>();
+			GeometryFrustumSelection->BindSelectionManager(WeakThis->SelectionManager);
+
+			USkeletalMeshModelingToolsGeometryClickSelection* GeometryClickSelection = Source->AddInteraction<
+				USkeletalMeshModelingToolsGeometryClickSelection>();
+			GeometryClickSelection->BindSelectionManager(WeakThis->SelectionManager);
+
+			USkeletalMeshModelingToolsGeometryHoverInteraction* GeometryHoverInteraction = Source->AddInteraction<
+				USkeletalMeshModelingToolsGeometryHoverInteraction>();
+			GeometryHoverInteraction->BindSelectionManager(WeakThis->SelectionManager);
+				
+			// Since viewport interactions are rebuilt whenever an editor mode is changed (including modes that can coexists like the persona skeleton selection edit mode)
+			// we have to check again to make sure the correct interactions are enabled when that happens
+			WeakThis->ToggleGeometrySelectionViewportInteractions(WeakThis->IsGeometrySelectionActive());
+		});
+	}
+
 	
 	CacheChangeCountWatcher.Initialize(
 		[this]()
@@ -393,19 +427,20 @@ void USkeletalMeshModelingToolsEditorMode::Enter()
 		[&](int32 ChangeCount)
 			{
 				TypedToolkit.Pin()->RefreshMorphTargetManager();
+				SelectionTransformTweaker->UpdateSelectionFrameTransform();
 			},
 		0);
 	
-	SkeletonTreeHostStateWatcher.Initialize(
+	SkeletonTreeHostStateUpdater.Initialize(
 		[this]()
 			{
-				FSkeletonTreeHostState State;
+				FSkeletonTreeHostStateFactors State;
 				State.bEnableDynamicMeshSkeleton = GetCurrentEditingCache() ? GetCurrentEditingCache()->IsDynamicMeshSkeletonEnabled() : false;
 				State.bIsUsingToolSkeletonWidget = bIsUsingToolSkeletonWidget;
 				State.bActiveToolModifiesSkeleton = bActiveToolModifiesSkeleton;
 				return State;
 			},
-		[this](const FSkeletonTreeHostState& State)
+		[this](const FSkeletonTreeHostStateFactors& State)
 			{
 				if (TypedToolkit.IsValid())
 				{
@@ -442,7 +477,7 @@ void USkeletalMeshModelingToolsEditorMode::Enter()
 					}
 				}
 			},
-		FSkeletonTreeHostState());
+		FSkeletonTreeHostStateFactors());
 	
 	CacheSkeletonChangeCountWatcher.Initialize(
 		[this]()
@@ -462,7 +497,50 @@ void USkeletalMeshModelingToolsEditorMode::Enter()
 				}
 			},
 		0);
+
+	GeometrySelectionStateWatcher.Initialize(
+		[this]()
+			{
+				return IsGeometrySelectionActive();
+			},
+		[this](bool InActive)
+			{
+				if (GetCurrentEditingCache())
+				{
+					GetCurrentEditingCache()->ToggleForceEnableDynamicMesh(InActive);
+				}
+				DesiredSkeletonState.bManipulation = !InActive;
+				ToggleGeometrySelectionViewportInteractions(InActive);
+			},
+		IsGeometrySelectionActive());
 	
+	SkeletonStateUpdater.Initialize(
+		[this]()
+			{
+				return GetSkeletonStateFactors();
+			},
+		[this](FSkeletonStateFactors InFactors)
+			{
+				if (GetCurrentEditingCache())
+				{
+					const bool bManipulation = !InFactors.bIsGeometrySelectionActive && InFactors.DesiredSkeletonState.bManipulation;
+					GetCurrentEditingCache()->ToggleBoneManipulation(bManipulation);
+
+					ESkeletonDrawMode SkeletonDrawMode = ESkeletonDrawMode::Hidden; 
+					if (InFactors.DesiredSkeletonState.bVisible)
+					{
+						SkeletonDrawMode = ESkeletonDrawMode::Default;
+						
+						if (InFactors.bIsGeometrySelectionActive)
+						{
+							SkeletonDrawMode = ESkeletonDrawMode::GreyedOut;
+						}
+					}
+					GetCurrentEditingCache()->SetSkeletonDrawMode(SkeletonDrawMode);
+					
+				}
+			},
+		GetSkeletonStateFactors());
 }
 
 UDebugSkelMeshComponent* USkeletalMeshModelingToolsEditorMode::GetSkelMeshComponent() const
@@ -521,10 +599,7 @@ void USkeletalMeshModelingToolsEditorMode::UnbindToolSkeletonTree()
 
 void USkeletalMeshModelingToolsEditorMode::ToggleBoneManipulation(bool bEnable)
 {
-	if (GetCurrentEditingCache())
-	{
-		GetCurrentEditingCache()->ToggleBoneManipulation(bEnable);
-	}
+	DesiredSkeletonState.bManipulation = bEnable;
 }
 
 USkeletonModifier* USkeletalMeshModelingToolsEditorMode::GetSkeletonReader()
@@ -639,17 +714,16 @@ void USkeletalMeshModelingToolsEditorMode::HandleSetMorphTargetAutoFill(FName Mo
 
 void USkeletalMeshModelingToolsEditorMode::Exit()
 {
-	RemoveViewportToolbarExtensions();
-	
-	// shutdown selection interaction
-	if (SelectionInteraction != nullptr)
+	if (UModeManagerInteractiveToolsContext* const ModeManagerInteractiveToolsContext = GetModeManager()->GetInteractiveToolsContext())
 	{
-		SelectionInteraction->Shutdown();
-		GetInteractiveToolsContext()->InputRouter->ForceTerminateSource(SelectionInteraction);
-		GetInteractiveToolsContext()->InputRouter->DeregisterSource(SelectionInteraction);
-		SelectionInteraction = nullptr;
+		ModeManagerInteractiveToolsContext->OnBuildViewportInteractions().Remove(BuildViewportInteractionsDelegateHandle);
 	}
 	
+	RemoveViewportToolbarExtensions();
+
+
+	SelectionTransformTweaker->Shutdown();
+	
 	// Shutdown SelectionManager. Wait until after Tool shutdown in case some restore-selection
 	// is involved (although since we are exiting Mode this currently would never matter)
 	if (SelectionManager != nullptr)
@@ -800,6 +874,10 @@ void USkeletalMeshModelingToolsEditorMode::InitializeGeometrySelectionSystems()
 		SelectorFactory->Init(this);
 		SelectionManager->RegisterSelectorFactory(MoveTemp(SelectorFactory));
 	}
+
+	SelectionTransformTweaker= NewObject<USkeletalMeshGeometrySelectionTransformTweaker>();
+	SelectionTransformTweaker->Setup(SelectionManager);
+	
 	const UModelingToolsModeCustomizationSettings* ModelingEditorSettings = GetDefault<UModelingToolsModeCustomizationSettings>();
 	
 	// Colors initialized here any time Modeling mode is entered
@@ -828,72 +906,10 @@ void USkeletalMeshModelingToolsEditorMode::InitializeGeometrySelectionSystems()
 
 
 
-	// Set up the selection interaction
-	SelectionInteraction = NewObject<UModelingSelectionInteraction>(GetToolManager());
-	SelectionInteraction->Initialize(SelectionManager,
-		[WeakThis]()
-			{
-				if (!WeakThis.IsValid())
-				{
-					return false;
-				}
-				return WeakThis->bEnableSelection;
-			},
-		[WeakThis](const FInputDeviceRay& DeviceRay)
-			{
-				// Because the editor gizmo does not participate in InputRouter behavior system, in some input behaviors
-				// we need to filter out clicks on the gizmo. This function can do this check.
-				if ( WeakThis.IsValid() && WeakThis->ShouldDrawWidget() )
-				{
-					if (FViewport* FocusedViewport = WeakThis->GetToolManager()->GetContextQueriesAPI()->GetFocusedViewport())
-					{
-						HHitProxy* HitResult = FocusedViewport->GetHitProxy(DeviceRay.ScreenPosition.X, DeviceRay.ScreenPosition.Y);
-						if (HitResult && HitResult->IsA(HWidgetAxis::StaticGetType()))
-						{
-							return true;
-						}
-					}
-				}
-				return false;
-			});
-	GetInteractiveToolsContext()->InputRouter->RegisterSource(SelectionInteraction);
 
 	GetInteractiveToolsContext()->OnRender.AddUObject(this, &USkeletalMeshModelingToolsEditorMode::OnToolContextRender);
 	GetInteractiveToolsContext()->OnDrawHUD.AddUObject(this, &USkeletalMeshModelingToolsEditorMode::OnToolContextDrawHUD);
 
-	// Register Drag Modes
-	auto MakeDragModeOptionAction = [WeakThis](EModelingSelectionInteraction_DragMode DragMode)
-		{
-			return FUIAction(
-				FExecuteAction::CreateLambda([WeakThis, DragMode]
-				{
-					if (!WeakThis.IsValid())
-					{
-						return;
-					}
-						
-					WeakThis->SelectionInteraction->SetActiveDragMode(DragMode);
-					UModelingToolsModeCustomizationSettings* ModelingEditorSettings = GetMutableDefault<UModelingToolsModeCustomizationSettings>();
-					ModelingEditorSettings->LastMeshSelectionDragMode = static_cast<int>(DragMode);
-					ModelingEditorSettings->SaveConfig();
-				}),
-				FCanExecuteAction(),
-				FIsActionChecked::CreateLambda([WeakThis, DragMode]()
-				{
-					if (!WeakThis.IsValid())
-					{
-						return false;
-					}
-
-					return WeakThis->SelectionInteraction->GetActiveDragMode() == DragMode;
-				}));
-		};
-
-	// Map drag mode commands
-	CommandList->MapAction(ToolManagerCommands.SelectionDragMode_None,
-		MakeDragModeOptionAction(EModelingSelectionInteraction_DragMode::NoDragInteraction));
-	CommandList->MapAction(ToolManagerCommands.SelectionDragMode_Path,
-		MakeDragModeOptionAction(EModelingSelectionInteraction_DragMode::PathInteraction));
 
 	// Register Frame Modes
 	auto ToggleFromGeometryAction = [WeakThis](EModelingSelectionInteraction_LocalFrameMode LocalFrameMode)
@@ -905,7 +921,7 @@ void USkeletalMeshModelingToolsEditorMode::InitializeGeometrySelectionSystems()
 					{
 						return;
 					}
-					WeakThis->SelectionInteraction->SetLocalFrameMode(LocalFrameMode);
+					WeakThis->SelectionTransformTweaker->SetLocalFrameMode(LocalFrameMode);
 					UModelingToolsModeCustomizationSettings* ModelingEditorSettings = GetMutableDefault<UModelingToolsModeCustomizationSettings>();
 					ModelingEditorSettings->LastMeshSelectionLocalFrameMode = static_cast<int>(LocalFrameMode);
 					ModelingEditorSettings->SaveConfig();
@@ -917,7 +933,7 @@ void USkeletalMeshModelingToolsEditorMode::InitializeGeometrySelectionSystems()
 					{
 						return false;
 					}
-					return WeakThis->SelectionInteraction->GetLocalFrameMode() == LocalFrameMode;
+					return WeakThis->SelectionTransformTweaker->GetLocalFrameMode() == LocalFrameMode;
 				}));
 		};
 	
@@ -954,9 +970,11 @@ void USkeletalMeshModelingToolsEditorMode::InitializeGeometrySelectionSystems()
 			FExecuteAction::CreateLambda([WeakThis, TopoMode, ElementMode]() {
 				if ( WeakThis.IsValid() && WeakThis->GetToolManager() && WeakThis->GetToolManager()->GetContextTransactionsAPI())
 				{
-					WeakThis->GetToolManager()->GetContextTransactionsAPI()->BeginUndoTransaction(LOCTEXT("ChangeSelectionMode", "Selection Mode"));
+					IToolsContextTransactionsAPI* TransactionsAPI = WeakThis->GetToolManager()->GetContextTransactionsAPI();
+					TransactionsAPI->BeginUndoTransaction(LOCTEXT("ChangeSelectionMode", "Selection Mode"));
 					WeakThis->SelectionManager->SetMeshSelectionTypeAndMode(ElementMode, TopoMode,  TopoMode != UGeometrySelectionManager::EMeshTopologyMode::None);
-					WeakThis->GetToolManager()->GetContextTransactionsAPI()->EndUndoTransaction();
+					TransactionsAPI->EndUndoTransaction();
+					
 					if (UModelingToolsModeCustomizationSettings* ModelingEditorSettings = GetMutableDefault<UModelingToolsModeCustomizationSettings>())
 					{
 						ModelingEditorSettings->LastMeshSelectionTopologyMode = static_cast<int>(TopoMode);
@@ -964,11 +982,7 @@ void USkeletalMeshModelingToolsEditorMode::InitializeGeometrySelectionSystems()
 						ModelingEditorSettings->SaveConfig();
 					}
 					
-					if (WeakThis->GetCurrentEditingCache())
-					{
-						const bool bElementSelectionEnabled = TopoMode != UGeometrySelectionManager::EMeshTopologyMode::None;
-						WeakThis->GetCurrentEditingCache()->ToggleForceEnableDynamicMesh(bElementSelectionEnabled);
-					}
+					
 				}
 			}),
 			FCanExecuteAction::CreateLambda([WeakThis]()
@@ -977,7 +991,7 @@ void USkeletalMeshModelingToolsEditorMode::InitializeGeometrySelectionSystems()
 				{
 					return false;
 				}
-				return WeakThis->bEnableSelection;
+				return WeakThis->bAllowGeometrySelection;
 			}),
 			FIsActionChecked::CreateLambda([WeakThis, TopoMode, ElementMode]()
 				{
@@ -1061,18 +1075,6 @@ void USkeletalMeshModelingToolsEditorMode::PopulateViewportToolbar()
 						FModelingToolsManagerCommands::Get().BeginSelectionAction_Contract, CommandList));
 				}
 
-				// Section for Drag Mode setting
-				{
-					FToolMenuSection& DragModeSection =
-						Submenu->FindOrAddSection("Drag Mode", LOCTEXT("DragModeLabel", "Drag Mode"));
-					DragModeSection.AddEntry(
-						FToolMenuEntry::InitMenuEntryWithCommandList(
-							FModelingToolsManagerCommands::Get().SelectionDragMode_None, CommandList));
-					DragModeSection.AddEntry(
-						FToolMenuEntry::InitMenuEntryWithCommandList(
-							FModelingToolsManagerCommands::Get().SelectionDragMode_Path, CommandList));
-				}
-
 				
 				// Section for Local Frame Mode setting
 				{
@@ -1121,6 +1123,41 @@ void USkeletalMeshModelingToolsEditorMode::RemoveViewportToolbarExtensions()
 	UToolMenus::Get()->UnregisterOwnerByName(SkeletalMeshModelingToolsEditorModeLocals::ViewportToolbarExtensionOwnerName);
 }
 
+bool USkeletalMeshModelingToolsEditorMode::IsGeometrySelectionActive() const
+{
+	return bAllowGeometrySelection && SelectionManager->GetMeshTopologyMode() != UGeometrySelectionManager::EMeshTopologyMode::None;
+}
+
+bool USkeletalMeshModelingToolsEditorMode::HasActiveGeometrySelection() const
+{
+	return IsGeometrySelectionActive() && SelectionManager->HasSelection();
+}
+
+void USkeletalMeshModelingToolsEditorMode::ToggleGeometrySelectionViewportInteractions(bool bEnable)
+{
+	if (GetModeManager()->GetInteractiveToolsContext()->UsesViewportInteractions())
+	{
+		UViewportInteractionsBehaviorSource* BehaviorSource = GetModeManager()->GetInteractiveToolsContext()->GetViewportInteractionsBehaviorSource();
+		
+		BehaviorSource->FindInteraction(USkeletalMeshModelingToolsGeometryFrustumSelection::StaticClass())->SetEnabled(bEnable);
+		BehaviorSource->FindInteraction(USkeletalMeshModelingToolsGeometryClickSelection::StaticClass())->SetEnabled(bEnable);
+		BehaviorSource->FindInteraction(USkeletalMeshModelingToolsGeometryHoverInteraction::StaticClass())->SetEnabled(bEnable);
+				
+		BehaviorSource->FindInteraction(USkeletalMeshModelingToolsClickSelection::StaticClass())->SetEnabled(!bEnable);
+		BehaviorSource->FindInteraction(UViewportMoveYawInteraction::StaticClass())->SetEnabled(!bEnable);	
+	}	
+}
+
+
+USkeletalMeshModelingToolsEditorMode::FSkeletonStateFactors USkeletalMeshModelingToolsEditorMode::GetSkeletonStateFactors()
+{
+	FSkeletonStateFactors Factors;
+	Factors.bIsGeometrySelectionActive = IsGeometrySelectionActive();
+	Factors.DesiredSkeletonState = DesiredSkeletonState;
+
+	return Factors;
+}
+
 void USkeletalMeshModelingToolsEditorMode::CreateToolkit()
 {
 	TSharedPtr<FSkeletalMeshModelingToolsEditorModeToolkit> NewToolKit = MakeShareable(new FSkeletalMeshModelingToolsEditorModeToolkit);
@@ -1148,8 +1185,10 @@ void USkeletalMeshModelingToolsEditorMode::Tick(FEditorViewportClient* ViewportC
 	}
 
 	CacheChangeCountWatcher.CheckAndUpdate();
-	SkeletonTreeHostStateWatcher.CheckAndUpdate();
+	SkeletonTreeHostStateUpdater.CheckAndUpdate();
 	CacheSkeletonChangeCountWatcher.CheckAndUpdate();
+	GeometrySelectionStateWatcher.CheckAndUpdate();
+	SkeletonStateUpdater.CheckAndUpdate();
 }
 
 bool USkeletalMeshModelingToolsEditorMode::HandleClick(FEditorViewportClient* InViewportClient, HHitProxy *HitProxy, const FViewportClick &Click)
@@ -1266,7 +1305,7 @@ void USkeletalMeshModelingToolsEditorMode::OnToolPostBuild( UInteractiveToolMana
 	ToggleBoneManipulation(false);
 
 	// Disable element selection modes while tool is active
-	bEnableSelection = false;
+	bAllowGeometrySelection = false;
 }
 
 void USkeletalMeshModelingToolsEditorMode::OnToolEndedWithStatus(UInteractiveToolManager* Manager, UInteractiveTool* Tool,
@@ -1281,29 +1320,17 @@ void USkeletalMeshModelingToolsEditorMode::OnToolEndedWithStatus(UInteractiveToo
 	
 	ToggleBoneManipulation(true);
 
-	bEnableSelection = true;
+	bAllowGeometrySelection = true;
 }
 
 void USkeletalMeshModelingToolsEditorMode::OnToolContextRender(IToolsContextRenderAPI* ToolsContextRenderAPI)
 {
-	if (SelectionInteraction)
-	{
-		SelectionInteraction->Render(ToolsContextRenderAPI);
-
-		// Bake in transform changes. Note that if we do this in OnToolsContextTick, it will 
-		// still block rendering updates if it is too expensive, unless it is only done every second Tick
-		SelectionInteraction->ApplyPendingTransformInteractions();
-	}
-	
 	SelectionManager->DebugRender(ToolsContextRenderAPI);
 }
 
 void USkeletalMeshModelingToolsEditorMode::OnToolContextDrawHUD(FCanvas* Canvas, IToolsContextRenderAPI* ToolsContextRenderAPI)
 {
-	if (SelectionInteraction)
-	{
-		SelectionInteraction->DrawHUD(Canvas, ToolsContextRenderAPI);
-	}	
+
 }
 
 bool USkeletalMeshModelingToolsEditorMode::ShouldToolStartBeAllowed(const FString& ToolIdentifier) const
@@ -1313,52 +1340,22 @@ bool USkeletalMeshModelingToolsEditorMode::ShouldToolStartBeAllowed(const FStrin
 	return true;
 }
 
-bool USkeletalMeshModelingToolsEditorMode::StartTracking(FEditorViewportClient* InViewportClient, FViewport* InViewport)
-{
-	if (GetCurrentEditingCache() && GetCurrentEditingCache()->IsDynamicMeshBoneManipulationEnabled())
-	{
-		if (GetCurrentEditingCache()->GetFirstSelectedBoneIndex() != INDEX_NONE)
-		{
-			const UE::Widget::EWidgetMode WidgetMode = InViewportClient->GetWidgetMode();
 
-			// we also allow undo/redo of bone manipulations
-			if (WidgetMode == UE::Widget::WM_Rotate)
-			{
-				GEditor->BeginTransaction(LOCTEXT("RotateBone", "Rotate Bone"));
-			}
-			else
-			{
-				GEditor->BeginTransaction(LOCTEXT("TranslateBone", "Translate Bone"));
-			}
-
-			GetCurrentEditingCache()->GetSkeletonPoser()->BeginPoseChange();
-			return true;
-		}
-	}
-	
-	return false;
-}
-
-bool USkeletalMeshModelingToolsEditorMode::EndTracking(FEditorViewportClient* InViewportClient, FViewport* InViewport)
-{
-	if (GetCurrentEditingCache() && GetCurrentEditingCache()->GetSkeletonPoser()->IsRecordingPoseChange())
-	{
-		GetCurrentEditingCache()->GetSkeletonPoser()->EndPoseChange();
-		
-		GEditor->EndTransaction();
-		return true;
-	}
-	return false;
-}
 
 bool USkeletalMeshModelingToolsEditorMode::InputDelta(FEditorViewportClient* InViewportClient, FViewport* InViewport, FVector& InDrag,
 	FRotator& InRot, FVector& InScale)
 {
+	const ECoordSystem CoordSystem = InViewportClient->GetWidgetCoordSystemSpace();	
+	if (SelectionTransformTweaker->IsEditingTransform())
+	{
+		SelectionTransformTweaker->TweakTransform(InDrag, InRot, InScale, CoordSystem == COORD_World);
+	}
+	
 	if (GetCurrentEditingCache() && GetCurrentEditingCache()->GetSkeletonPoser()->IsRecordingPoseChange())
 	{
 		// Actually change bone transform
 		const EAxisList::Type CurrentAxis = InViewportClient->GetCurrentWidgetAxis();
-		const ECoordSystem CoordSystem = InViewportClient->GetWidgetCoordSystemSpace();
+		
 		if (CurrentAxis != EAxisList::Type::None)
 		{
 			URefSkeletonPoser* Poser = GetCurrentEditingCache()->GetSkeletonPoser();
@@ -1427,6 +1424,11 @@ void USkeletalMeshModelingToolsEditorMode::DrawHUD(FEditorViewportClient* Viewpo
 
 bool USkeletalMeshModelingToolsEditorMode::AllowWidgetMove()
 {
+	if (HasActiveGeometrySelection())
+	{
+		return true;
+	}
+	
 	if (GetCurrentEditingCache() && GetCurrentEditingCache()->IsDynamicMeshSkeletonEnabled())
 	{
 		return ShouldDrawWidget();
@@ -1436,6 +1438,11 @@ bool USkeletalMeshModelingToolsEditorMode::AllowWidgetMove()
 
 bool USkeletalMeshModelingToolsEditorMode::ShouldDrawWidget() const
 {
+	if (HasActiveGeometrySelection())
+	{
+		return true;
+	}
+	
 	if (GetCurrentEditingCache() && GetCurrentEditingCache()->IsDynamicMeshSkeletonEnabled())
 	{
 		if (GetCurrentEditingCache()->IsDynamicMeshBoneManipulationEnabled() && GetCurrentEditingCache()->GetFirstSelectedBoneIndex() != INDEX_NONE )
@@ -1452,6 +1459,10 @@ bool USkeletalMeshModelingToolsEditorMode::ShouldDrawWidget() const
 bool USkeletalMeshModelingToolsEditorMode::UsesTransformWidget() const
 {
 	// This function determines if the GetWidgetLocation() of this editor mode should be used
+	if (HasActiveGeometrySelection())
+	{
+		return true;
+	}
 	
 	if (GetCurrentEditingCache() && GetCurrentEditingCache()->IsDynamicMeshSkeletonEnabled())
 	{
@@ -1463,6 +1474,11 @@ bool USkeletalMeshModelingToolsEditorMode::UsesTransformWidget() const
 
 FVector USkeletalMeshModelingToolsEditorMode::GetWidgetLocation() const
 {
+	if (HasActiveGeometrySelection())
+	{
+		return SelectionTransformTweaker->GetSelectionFrameTransform().GetLocation();
+	}
+	
 	if (GetCurrentEditingCache() && GetCurrentEditingCache()->IsDynamicMeshSkeletonEnabled())
 	{
 		int32 BoneIndex = GetCurrentEditingCache()->GetFirstSelectedBoneIndex();
@@ -1487,6 +1503,12 @@ FVector USkeletalMeshModelingToolsEditorMode::GetWidgetLocation() const
 
 bool USkeletalMeshModelingToolsEditorMode::GetCustomDrawingCoordinateSystem(FMatrix& InMatrix, void* InData)
 {
+	if (HasActiveGeometrySelection())
+	{
+		InMatrix = SelectionTransformTweaker->GetSelectionFrameTransform().ToMatrixNoScale().RemoveTranslation();
+		return true;
+	}
+	
 	if (GetCurrentEditingCache() && GetCurrentEditingCache()->IsDynamicMeshSkeletonEnabled())
 	{
 		const bool bIsParentMode = Owner ? Owner->GetCoordSystem() == COORD_Parent : false;
@@ -1520,6 +1542,11 @@ bool USkeletalMeshModelingToolsEditorMode::GetCustomDrawingCoordinateSystem(FMat
 
 bool USkeletalMeshModelingToolsEditorMode::GetCustomInputCoordinateSystem(FMatrix& InMatrix, void* InData)
 {
+	if (HasActiveGeometrySelection())
+	{
+		return GetCustomDrawingCoordinateSystem(InMatrix, InData);	
+	}
+	
 	if (GetCurrentEditingCache() && GetCurrentEditingCache()->IsDynamicMeshSkeletonEnabled())
 	{
 		return GetCustomDrawingCoordinateSystem(InMatrix, InData);
@@ -1543,6 +1570,44 @@ bool USkeletalMeshModelingToolsEditorMode::ShouldCommitToSkeletalMeshOnToolCommi
 	return ApplyMode == EApplyMode::ApplyOnToolExit;
 }
 
+bool USkeletalMeshModelingToolsEditorMode::BeginTransform(const FGizmoState& InState)
+{
+	if (HasActiveGeometrySelection())
+	{
+		SelectionTransformTweaker->BeginTransformEdit();
+	}
+	
+	if (GetCurrentEditingCache() && GetCurrentEditingCache()->IsDynamicMeshBoneManipulationEnabled())
+	{
+		if (GetCurrentEditingCache()->GetFirstSelectedBoneIndex() != INDEX_NONE)
+		{
+			GetCurrentEditingCache()->GetSkeletonPoser()->BeginPoseChange();
+			return true;
+		}
+	}
+	
+	return false;
+}
+
+bool USkeletalMeshModelingToolsEditorMode::EndTransform(const FGizmoState& InState)
+{
+	if (SelectionTransformTweaker->IsEditingTransform())
+	{
+		SelectionTransformTweaker->EndTransformEdit();
+
+		return true;
+	}
+	
+	if (GetCurrentEditingCache() && GetCurrentEditingCache()->GetSkeletonPoser()->IsRecordingPoseChange())
+	{
+		GetCurrentEditingCache()->GetSkeletonPoser()->EndPoseChange();
+		
+		return true;
+	}
+	return false;	
+}
+
+
 TSharedPtr<ISkeletalMeshEditorBinding> USkeletalMeshModelingToolsEditorMode::GetModeBinding()
 {
 	if (!ModeBinding.IsValid())
@@ -1774,20 +1839,12 @@ void USkeletalMeshModelingToolsEditorMode::SetApplyMode(EApplyMode InApplyMode)
 
 void USkeletalMeshModelingToolsEditorMode::HideSkeletonForTool()
 {
-	if (!GetCurrentEditingCache())
-	{
-		return;
-	}	
-	GetCurrentEditingCache()->HideSkeleton();
+	DesiredSkeletonState.bVisible = false;
 }
 
 void USkeletalMeshModelingToolsEditorMode::ShowSkeletonForTool()
 {
-	if (!GetCurrentEditingCache())
-	{
-		return;
-	}
-	GetCurrentEditingCache()->ShowSkeleton();
+	DesiredSkeletonState.bVisible = true;
 }
 
 ISkeletalMeshEditingInterface* USkeletalMeshModelingToolsEditorMode::GetSkeletonInterface(UInteractiveTool* InTool)
diff --git a/Engine/Plugins/Animation/SkeletalMeshModelingTools/Source/SkeletalMeshModelingTools/Private/SkeletalMeshModelingToolsEditorMode.h b/Engine/Plugins/Animation/SkeletalMeshModelingTools/Source/SkeletalMeshModelingTools/Private/SkeletalMeshModelingToolsEditorMode.h
--- a/Engine/Plugins/Animation/SkeletalMeshModelingTools/Source/SkeletalMeshModelingTools/Private/SkeletalMeshModelingToolsEditorMode.h
+++ b/Engine/Plugins/Animation/SkeletalMeshModelingTools/Source/SkeletalMeshModelingTools/Private/SkeletalMeshModelingToolsEditorMode.h
@@ -4,17 +4,18 @@
 
 #include "CoreMinimal.h"
 #include "EdMode.h"
+#include "GizmoEdModeInterface.h"
+#include "ModelingSelectionInteraction.h"
 #include "SkeletalMeshModelingModeToolExtensions.h"
 #include "SkeletalMeshNotifier.h"
 #include "ToolContextInterfaces.h"
 #include "Interfaces/ISKMBackedDynaMeshComponentProvider.h"
-#include "Animation/DebugSkelMeshComponent.h"
 #include "Changes/ValueWatcher.h"
-#include "Misc/Change.h"
 #include "TargetInterfaces/MeshTargetInterfaceTypes.h"
 
 #include "SkeletalMeshModelingToolsEditorMode.generated.h"
 
+class USkeletalMeshGeometrySelectionTransformTweaker;
 class FSkeletalMeshModelingToolsEditorModeToolkit;
 
 class USkeletalMeshEditingCache;
@@ -35,6 +36,7 @@ class USkeletonModifier;
 class UGeometrySelectionManager;
 class UModelingSelectionInteraction;
 class UInteractiveCommand;
+class USkeletalMesh;
 
 namespace UE
 {
@@ -71,7 +73,8 @@ private:
 UCLASS()
 class USkeletalMeshModelingToolsEditorMode : 
 	public UBaseLegacyWidgetEdMode,
-	public ISkeletalMeshBackedDynamicMeshComponentProvider 
+	public ISkeletalMeshBackedDynamicMeshComponentProvider,
+	public IGizmoEdModeInterface
 {
 	GENERATED_BODY()
 public:
@@ -94,8 +97,8 @@ public:
 	virtual bool ComputeBoundingBoxForViewportFocus(AActor* Actor, UPrimitiveComponent* PrimitiveComponent, FBox& InOutBox) const override;
 	virtual bool UsesToolkits() const override { return true; }
 	virtual bool ShouldToolStartBeAllowed(const FString& ToolIdentifier) const override;
-	virtual bool StartTracking(FEditorViewportClient* InViewportClient, FViewport* InViewport) override;
-	virtual bool EndTracking(FEditorViewportClient* InViewportClient, FViewport* InViewport) override;
+	// virtual bool StartTracking(FEditorViewportClient* InViewportClient, FViewport* InViewport) override;
+	// virtual bool EndTracking(FEditorViewportClient* InViewportClient, FViewport* InViewport) override;
 	virtual bool InputDelta(FEditorViewportClient* InViewportClient, FViewport* InViewport, FVector& InDrag, FRotator& InRot, FVector& InScale) override;
 	virtual void DrawHUD(FEditorViewportClient* ViewportClient, FViewport* Viewport, const FSceneView* View, FCanvas* Canvas) override;
 	virtual bool AllowWidgetMove() override;
@@ -111,6 +114,10 @@ public:
 	// ISkeletalMeshBackedDynamicMeshComponentProvider
 	virtual USkeletalMeshBackedDynamicMeshComponent* GetComponent(UObject* SourceObject) override;
 	bool ShouldCommitToSkeletalMeshOnToolCommit() override;
+
+	// IGizmoEdModeInterface
+	virtual bool BeginTransform(const FGizmoState& InState) override;
+	virtual bool EndTransform(const FGizmoState& InState) override;
 	
 	// binding
 	TSharedPtr<ISkeletalMeshEditorBinding> GetModeBinding();
@@ -214,22 +221,31 @@ private:
 	void ToggleSkeletalMeshBoneManipulation(bool bEnable);
 	bool IsSkeletalMeshBoneManipulationEnabled();
 
+	struct FDesiredSkeletonState
+	{
+		bool bVisible = true;
+		bool bManipulation = true;
+		friend bool operator==(const FDesiredSkeletonState&, const FDesiredSkeletonState&) = default;
+	};
+
+	FDesiredSkeletonState DesiredSkeletonState;
+	
 	// A dummy read-only modifier to host the current edited skeleton for RefSkeletonTree UI
 	UPROPERTY(NonTransactional)
 	TObjectPtr<USkeletonModifier> SkeletonReader = nullptr;
 
 	TValueWatcher<int32> CacheChangeCountWatcher;
 	
-	struct FSkeletonTreeHostState
+	struct FSkeletonTreeHostStateFactors
 	{
 		bool bEnableDynamicMeshSkeleton = false;
 		bool bIsUsingToolSkeletonWidget = false;
 		bool bActiveToolModifiesSkeleton = false;
-		friend bool operator==(const FSkeletonTreeHostState&, const FSkeletonTreeHostState&) = default;
+		friend bool operator==(const FSkeletonTreeHostStateFactors&, const FSkeletonTreeHostStateFactors&) = default;
 	};
 	
 	// Decides which skeleton tree widget to show to the user
-	TValueWatcher<FSkeletonTreeHostState> SkeletonTreeHostStateWatcher;
+	TValueWatcher<FSkeletonTreeHostStateFactors> SkeletonTreeHostStateUpdater;
 	TValueWatcher<int32> CacheSkeletonChangeCountWatcher;
 
 	bool bIsUsingToolSkeletonWidget = false;
@@ -257,14 +273,34 @@ private:
 	TObjectPtr<UGeometrySelectionManager> SelectionManager;
 
 	UPROPERTY(NonTransactional)
-	TObjectPtr<UModelingSelectionInteraction> SelectionInteraction;
+	TObjectPtr<USkeletalMeshGeometrySelectionTransformTweaker> SelectionTransformTweaker;
 
-	bool bEnableSelection = true;
+	bool bAllowGeometrySelection = true;
 
 	void InitializeGeometrySelectionSystems();
 	void PopulateViewportToolbar();
 	void RemoveViewportToolbarExtensions();
-	
+
+	// Used to keep the command objects alive
 	UPROPERTY(NonTransactional)
 	TArray<TObjectPtr<UInteractiveCommand>> ModelingModeCommands;
+
+	FDelegateHandle BuildViewportInteractionsDelegateHandle;
+
+	TValueWatcher<bool> GeometrySelectionStateWatcher;
+
+	bool IsGeometrySelectionActive() const;
+	bool HasActiveGeometrySelection() const;
+	void ToggleGeometrySelectionViewportInteractions(bool bEnable);
+
+	struct FSkeletonStateFactors
+	{
+		bool bIsGeometrySelectionActive = false;
+		FDesiredSkeletonState DesiredSkeletonState;
+		friend bool operator==(const FSkeletonStateFactors&, const FSkeletonStateFactors&) = default;
+	};
+	FSkeletonStateFactors GetSkeletonStateFactors();
+
+	TValueWatcher<FSkeletonStateFactors> SkeletonStateUpdater;
+	
 };
diff --git a/Engine/Plugins/Animation/SkeletalMeshModelingTools/Source/SkeletalMeshModelingTools/Private/ViewportInteractions/SKMModelingToolsClickSelection.cpp b/Engine/Plugins/Animation/SkeletalMeshModelingTools/Source/SkeletalMeshModelingTools/Private/ViewportInteractions/SKMModelingToolsClickSelection.cpp
new file mode 100644
--- /dev/null
+++ b/Engine/Plugins/Animation/SkeletalMeshModelingTools/Source/SkeletalMeshModelingTools/Private/ViewportInteractions/SKMModelingToolsClickSelection.cpp
@@ -0,0 +1,27 @@
+// Copyright Epic Games, Inc. All Rights Reserved.
+
+#include "ViewportInteractions/SKMModelingToolsClickSelection.h"
+
+#include "EditorModeManager.h"
+#include "EditorViewportClient.h"
+
+USkeletalMeshModelingToolsClickSelection::USkeletalMeshModelingToolsClickSelection()
+{
+	InteractionName = TEXT("Skeletal Mesh Modeling Tools Click Selection");
+	Groups = { UE::Editor::ViewportInteractions::ViewportClick };
+}
+
+FInputRayHit USkeletalMeshModelingToolsClickSelection::IsHitByClick(const FInputDeviceRay& ClickPos)
+{
+	if (IsEnabled())
+	{
+		return Super::IsHitByClick(ClickPos);
+	}
+	
+	return FInputRayHit();
+}
+
+void USkeletalMeshModelingToolsClickSelection::ProcessClick_Internal(const FViewportClick& InViewportClick, HHitProxy* InHitProxy, FSceneView& View)
+{
+	GetEditorViewportClient()->GetModeTools()->HandleClick(GetEditorViewportClient(),  InHitProxy, InViewportClick);
+}
diff --git a/Engine/Plugins/Animation/SkeletalMeshModelingTools/Source/SkeletalMeshModelingTools/Private/ViewportInteractions/SKMModelingToolsClickSelection.h b/Engine/Plugins/Animation/SkeletalMeshModelingTools/Source/SkeletalMeshModelingTools/Private/ViewportInteractions/SKMModelingToolsClickSelection.h
new file mode 100644
--- /dev/null
+++ b/Engine/Plugins/Animation/SkeletalMeshModelingTools/Source/SkeletalMeshModelingTools/Private/ViewportInteractions/SKMModelingToolsClickSelection.h
@@ -0,0 +1,22 @@
+// Copyright Epic Games, Inc. All Rights Reserved.
+
+#pragma once
+
+#include "ViewportInteractions/ViewportClickInteraction.h"
+#include "SKMModelingToolsClickSelection.generated.h"
+
+
+UCLASS(MinimalAPI, Transient)
+class USkeletalMeshModelingToolsClickSelection
+	: public UViewportClickInteraction
+{
+	GENERATED_BODY()
+
+public:
+	USkeletalMeshModelingToolsClickSelection();
+	virtual FInputRayHit IsHitByClick(const FInputDeviceRay& ClickPos) override;
+
+protected:
+	virtual void ProcessClick_Internal(const FViewportClick& InViewportClick, HHitProxy* InHitProxy, FSceneView& View) override;
+};
+
diff --git a/Engine/Plugins/Animation/SkeletalMeshModelingTools/Source/SkeletalMeshModelingTools/Private/ViewportInteractions/SKMModelingToolsGeometryClickSelection.cpp b/Engine/Plugins/Animation/SkeletalMeshModelingTools/Source/SkeletalMeshModelingTools/Private/ViewportInteractions/SKMModelingToolsGeometryClickSelection.cpp
new file mode 100644
--- /dev/null
+++ b/Engine/Plugins/Animation/SkeletalMeshModelingTools/Source/SkeletalMeshModelingTools/Private/ViewportInteractions/SKMModelingToolsGeometryClickSelection.cpp
@@ -0,0 +1,52 @@
+// Copyright Epic Games, Inc. All Rights Reserved.
+
+#include "ViewportInteractions/SKMModelingToolsGeometryClickSelection.h"
+
+#include "EditorModeManager.h"
+#include "EditorViewportClient.h"
+#include "BaseBehaviors/SingleClickBehavior.h"
+#include "Selection/GeometrySelectionManager.h"
+
+USkeletalMeshModelingToolsGeometryClickSelection::USkeletalMeshModelingToolsGeometryClickSelection()
+{
+	InteractionName = TEXT("Skeletal Mesh Modeling Tools Geometry Click Selection");
+	Groups = { UE::Editor::ViewportInteractions::ViewportClick };
+
+	USingleClickInputBehavior* SingleClickBehavior = NewObject<USingleClickInputBehavior>();
+	SingleClickBehavior->Initialize(this);
+
+	SingleClickBehavior->SetDefaultPriority(UE::Editor::ViewportInteractions::CLICK_PRIORITY);
+	SingleClickBehavior->SetAllowsCaptureStealing(true);
+	RegisterInputBehavior(SingleClickBehavior);
+}
+
+FInputRayHit USkeletalMeshModelingToolsGeometryClickSelection::IsHitByClick(const FInputDeviceRay& ClickPos)
+{
+	if (IsEnabled())
+	{
+		return FInputRayHit((double)TNumericLimits<float>::Max());
+	}
+	
+	return FInputRayHit();
+}
+
+void USkeletalMeshModelingToolsGeometryClickSelection::OnClicked(const FInputDeviceRay& ClickPos)
+{
+	UE::Geometry::FGeometrySelectionUpdateConfig UpdateConfig;
+	UpdateConfig.ChangeType = UE::Geometry::EGeometrySelectionChangeType::Replace;
+	if (IsShiftDown())
+	{
+		UpdateConfig.ChangeType = UE::Geometry::EGeometrySelectionChangeType::Add;
+	}
+	else if (IsCtrlDown())
+	{
+		UpdateConfig.ChangeType = UE::Geometry::EGeometrySelectionChangeType::Remove;
+	}
+	UE::Geometry::FGeometrySelectionUpdateResult Result;
+	SelectionManager->UpdateSelectionViaRaycast(ClickPos.WorldRay, UpdateConfig, Result);	
+}
+
+void USkeletalMeshModelingToolsGeometryClickSelection::BindSelectionManager(UGeometrySelectionManager* InSelectionManager)
+{
+	SelectionManager = InSelectionManager;
+}
diff --git a/Engine/Plugins/Animation/SkeletalMeshModelingTools/Source/SkeletalMeshModelingTools/Private/ViewportInteractions/SKMModelingToolsGeometryClickSelection.h b/Engine/Plugins/Animation/SkeletalMeshModelingTools/Source/SkeletalMeshModelingTools/Private/ViewportInteractions/SKMModelingToolsGeometryClickSelection.h
new file mode 100644
--- /dev/null
+++ b/Engine/Plugins/Animation/SkeletalMeshModelingTools/Source/SkeletalMeshModelingTools/Private/ViewportInteractions/SKMModelingToolsGeometryClickSelection.h
@@ -0,0 +1,30 @@
+// Copyright Epic Games, Inc. All Rights Reserved.
+
+#pragma once
+
+#include "ViewportInteractions/ViewportClickInteraction.h"
+#include "SKMModelingToolsGeometryClickSelection.generated.h"
+
+class UGeometrySelectionManager;
+
+UCLASS(MinimalAPI, Transient)
+class USkeletalMeshModelingToolsGeometryClickSelection
+	: public UViewportInteraction
+	, public IClickBehaviorTarget
+{
+	GENERATED_BODY()
+
+public:
+	USkeletalMeshModelingToolsGeometryClickSelection();
+	virtual FInputRayHit IsHitByClick(const FInputDeviceRay& ClickPos) override;
+	virtual void OnClicked(const FInputDeviceRay& ClickPos) override;
+
+	void BindSelectionManager(UGeometrySelectionManager* InSelectionManager);
+
+	virtual TArray<FEditorModeID> GetUnsupportedModes() const override {return {};}
+	
+protected:
+	
+	TWeakObjectPtr<UGeometrySelectionManager> SelectionManager;
+};
+
diff --git a/Engine/Plugins/Animation/SkeletalMeshModelingTools/Source/SkeletalMeshModelingTools/Private/ViewportInteractions/SKMModelingToolsGeometryFrustumSelection.cpp b/Engine/Plugins/Animation/SkeletalMeshModelingTools/Source/SkeletalMeshModelingTools/Private/ViewportInteractions/SKMModelingToolsGeometryFrustumSelection.cpp
new file mode 100644
--- /dev/null
+++ b/Engine/Plugins/Animation/SkeletalMeshModelingTools/Source/SkeletalMeshModelingTools/Private/ViewportInteractions/SKMModelingToolsGeometryFrustumSelection.cpp
@@ -0,0 +1,98 @@
+// Copyright Epic Games, Inc. All Rights Reserved.
+
+#include "ViewportInteractions/SKMModelingToolsGeometryFrustumSelection.h"
+
+#include "EditorModeManager.h"
+#include "EditorViewportClient.h"
+#include "SceneView.h"
+#include "ScopedTransaction.h"
+#include "Selection/GeometrySelectionManager.h"
+
+USkeletalMeshModelingToolsGeometryFrustumSelection::USkeletalMeshModelingToolsGeometryFrustumSelection()
+{
+	InteractionName = TEXT("Skeletal Mesh Modeling Tools Geometry Frustum Selection");
+	Groups = { UE::Editor::ViewportInteractions::FrustumSelect };
+
+	if (UViewportClickDragBehavior* ClickDragBehavior = ClickDragInputBehavior.Get())
+	{
+		ClickDragBehavior->SetBindings({
+			UE::Editor::ViewportInteractions::FButtonBinding(EKeys::LeftMouseButton).TriggersStart()
+		});
+	}
+}
+
+void USkeletalMeshModelingToolsGeometryFrustumSelection::BindSelectionManager(UGeometrySelectionManager* InSelectionManager)
+{
+	SelectionManager = InSelectionManager;
+}
+
+FInputRayHit USkeletalMeshModelingToolsGeometryFrustumSelection::CanBeginClickDragSequence(const FInputDeviceRay& PressPos)
+{
+	if (IsEnabled())
+	{
+		return Super::CanBeginClickDragSequence(PressPos);
+	}
+
+	return FInputRayHit(); 
+}
+
+void USkeletalMeshModelingToolsGeometryFrustumSelection::OnDragEnd(const FInputDeviceRay& InReleasePos)
+{
+	End = FVector(InReleasePos.ScreenPosition.X, InReleasePos.ScreenPosition.Y, 0);
+
+	if (!EditorViewportClientProxy)
+	{
+		return;
+	}
+
+	FEditorViewportClient* const EditorViewportClient = GetEditorViewportClient();
+	if (!EditorViewportClient)
+	{
+		return;
+	}
+
+	FViewport* Viewport = EditorViewportClient->Viewport;
+	if (!Viewport)
+	{
+		return;
+	}
+
+	FEditorModeTools* ModeTools = EditorViewportClient->GetModeTools();
+	if (!ModeTools)
+	{
+		return;
+	}
+
+	FSceneViewFamilyContext ViewFamily(
+		FSceneViewFamily::ConstructionValues(Viewport, EditorViewportClient->GetScene(), EditorViewportClient->EngineShowFlags)
+	);
+	FSceneView* SceneView = EditorViewportClient->CalcSceneView(&ViewFamily);
+
+	// Generate a frustum out of the dragged box
+	FConvexVolume Frustum;
+	constexpr bool bUseBoxFrustum = true;
+	UE::EditorDragTools::FViewportFrustum ViewportFrustum(*EditorViewportClient, *SceneView, Start, End, bUseBoxFrustum);
+	ViewportFrustum.Calculate(Frustum);
+
+	if (SelectionManager.IsValid()
+		&& SelectionManager->HasActiveTargets() 
+		&& SelectionManager->GetMeshTopologyMode() != UGeometrySelectionManager::EMeshTopologyMode::None)
+	{
+		UE::Geometry::FGeometrySelectionUpdateConfig UpdateConfig;
+		UpdateConfig.ChangeType = UE::Geometry::EGeometrySelectionChangeType::Replace;
+		if (FInputDeviceState::IsShiftKeyDown(InputState))
+		{
+			UpdateConfig.ChangeType = UE::Geometry::EGeometrySelectionChangeType::Add;
+		}
+		else if ( FInputDeviceState::IsCtrlKeyDown(InputState))
+		{
+			UpdateConfig.ChangeType = UE::Geometry::EGeometrySelectionChangeType::Remove;
+		}			
+
+		UE::Geometry::FGeometrySelectionUpdateResult Result;
+		SelectionManager->UpdateSelectionViaConvex(
+			Frustum, UpdateConfig, Result);
+	}
+
+	UDragToolInteraction::OnDragEnd(InReleasePos);
+}
diff --git a/Engine/Plugins/Animation/SkeletalMeshModelingTools/Source/SkeletalMeshModelingTools/Private/ViewportInteractions/SKMModelingToolsGeometryFrustumSelection.h b/Engine/Plugins/Animation/SkeletalMeshModelingTools/Source/SkeletalMeshModelingTools/Private/ViewportInteractions/SKMModelingToolsGeometryFrustumSelection.h
new file mode 100644
--- /dev/null
+++ b/Engine/Plugins/Animation/SkeletalMeshModelingTools/Source/SkeletalMeshModelingTools/Private/ViewportInteractions/SKMModelingToolsGeometryFrustumSelection.h
@@ -0,0 +1,28 @@
+// Copyright Epic Games, Inc. All Rights Reserved.
+
+#pragma once
+
+#include "EditorDragTools/FrustumSelectInteraction.h"
+#include "SKMModelingToolsGeometryFrustumSelection.generated.h"
+
+class UGeometrySelectionManager;
+
+UCLASS(MinimalAPI, Transient)
+class USkeletalMeshModelingToolsGeometryFrustumSelection : public UFrustumSelectInteraction
+{
+	GENERATED_BODY()
+
+public:
+	USkeletalMeshModelingToolsGeometryFrustumSelection();
+
+	void BindSelectionManager(UGeometrySelectionManager* InSelectionManager);
+
+	virtual FInputRayHit CanBeginClickDragSequence(const FInputDeviceRay& PressPos) override;
+	virtual void OnDragEnd(const FInputDeviceRay& InReleasePos) override;
+
+	virtual TArray<FEditorModeID> GetUnsupportedModes() const override {return {};}
+	
+protected:
+	TWeakObjectPtr<UGeometrySelectionManager> SelectionManager;
+};
+
diff --git a/Engine/Plugins/Animation/SkeletalMeshModelingTools/Source/SkeletalMeshModelingTools/Private/ViewportInteractions/SKMModelingToolsGeometryHoverInteraction.cpp b/Engine/Plugins/Animation/SkeletalMeshModelingTools/Source/SkeletalMeshModelingTools/Private/ViewportInteractions/SKMModelingToolsGeometryHoverInteraction.cpp
new file mode 100644
--- /dev/null
+++ b/Engine/Plugins/Animation/SkeletalMeshModelingTools/Source/SkeletalMeshModelingTools/Private/ViewportInteractions/SKMModelingToolsGeometryHoverInteraction.cpp
@@ -0,0 +1,49 @@
+// Copyright Epic Games, Inc. All Rights Reserved.
+
+#include "ViewportInteractions/SKMModelingToolsGeometryHoverInteraction.h"
+
+#include "BaseBehaviors/MouseHoverBehavior.h"
+#include "Selection/GeometrySelectionManager.h"
+
+
+USkeletalMeshModelingToolsGeometryHoverInteraction::USkeletalMeshModelingToolsGeometryHoverInteraction()
+{
+	HoverBehavior = NewObject<UMouseHoverBehavior>();
+	HoverBehavior->Initialize(this);
+	
+	RegisterInputBehavior(HoverBehavior);
+}
+
+void USkeletalMeshModelingToolsGeometryHoverInteraction::BindSelectionManager(UGeometrySelectionManager* InSelectionManager)
+{
+	SelectionManager = InSelectionManager;
+}
+
+FInputRayHit USkeletalMeshModelingToolsGeometryHoverInteraction::BeginHoverSequenceHitTest(const FInputDeviceRay& PressPos)
+{
+	if (!IsEnabled())
+	{
+		return {};
+	}
+
+	FInputRayHit ActiveObjectHit = FInputRayHit();
+	SelectionManager->RayHitTest(PressPos.WorldRay, ActiveObjectHit);
+	return ActiveObjectHit;
+}
+
+
+void USkeletalMeshModelingToolsGeometryHoverInteraction::OnBeginHover(const FInputDeviceRay& DevicePos)
+{
+	SelectionManager->UpdateSelectionPreviewViaRaycast(DevicePos.WorldRay);
+}
+
+bool USkeletalMeshModelingToolsGeometryHoverInteraction::OnUpdateHover(const FInputDeviceRay& DevicePos)
+{
+	return SelectionManager->UpdateSelectionPreviewViaRaycast(DevicePos.WorldRay);
+}
+
+void USkeletalMeshModelingToolsGeometryHoverInteraction::OnEndHover()
+{
+	SelectionManager->ClearSelectionPreview();
+}
+
diff --git a/Engine/Plugins/Animation/SkeletalMeshModelingTools/Source/SkeletalMeshModelingTools/Private/ViewportInteractions/SKMModelingToolsGeometryHoverInteraction.h b/Engine/Plugins/Animation/SkeletalMeshModelingTools/Source/SkeletalMeshModelingTools/Private/ViewportInteractions/SKMModelingToolsGeometryHoverInteraction.h
new file mode 100644
--- /dev/null
+++ b/Engine/Plugins/Animation/SkeletalMeshModelingTools/Source/SkeletalMeshModelingTools/Private/ViewportInteractions/SKMModelingToolsGeometryHoverInteraction.h
@@ -0,0 +1,51 @@
+// Copyright Epic Games, Inc. All Rights Reserved.
+
+#pragma once
+
+#include "SkeletalMeshModelingToolsEditorMode.h"
+#include "BaseBehaviors/BehaviorTargetInterfaces.h"
+#include "ViewportInteractions/ViewportInteraction.h"
+#include "SKMModelingToolsGeometryHoverInteraction.generated.h"
+
+
+
+UCLASS(MinimalAPI, Transient)
+class USkeletalMeshModelingToolsGeometryHoverInteraction
+	: public UViewportInteraction
+	, public IHoverBehaviorTarget
+{
+	GENERATED_BODY()
+
+public:
+	USkeletalMeshModelingToolsGeometryHoverInteraction();
+
+	void BindSelectionManager(UGeometrySelectionManager* InSelectionManager);
+	
+	/**
+	 * Do hover hit-test
+	 */
+	virtual FInputRayHit BeginHoverSequenceHitTest(const FInputDeviceRay& PressPos) override;
+
+	/**
+	 * Initialize hover sequence at given position
+	 */
+	virtual void OnBeginHover(const FInputDeviceRay& DevicePos) override;
+
+	/**
+	 * Update active hover sequence with new input position
+	 */
+	virtual bool OnUpdateHover(const FInputDeviceRay& DevicePos) override;
+
+	/**
+	 * Terminate active hover sequence
+	 */
+	virtual void OnEndHover() override;
+	
+protected:	
+	UPROPERTY()
+	TObjectPtr<UMouseHoverBehavior> HoverBehavior;
+
+	TWeakObjectPtr<UGeometrySelectionManager> SelectionManager;
+	
+};
+
