[SkeletalMeshModellingTools] Selecting in the list of weighted Bones in the Weight Tool now selects the Bone in the Skeleton Tree

#Synced-CL 49930900

diff --git a/Engine/Plugins/Animation/SkeletalMeshModelingTools/Source/SkeletalMeshModelingTools/Private/DetailCustomization/SkinWeightDetailCustomization.cpp b/Engine/Plugins/Animation/SkeletalMeshModelingTools/Source/SkeletalMeshModelingTools/Private/DetailCustomization/SkinWeightDetailCustomization.cpp
--- a/Engine/Plugins/Animation/SkeletalMeshModelingTools/Source/SkeletalMeshModelingTools/Private/DetailCustomization/SkinWeightDetailCustomization.cpp
+++ b/Engine/Plugins/Animation/SkeletalMeshModelingTools/Source/SkeletalMeshModelingTools/Private/DetailCustomization/SkinWeightDetailCustomization.cpp
@@ -12,6 +12,7 @@
 #include "Widgets/Input/SButton.h"
 #include "Widgets/Input/SNumericEntryBox.h"
 #include "Framework/MultiBox/MultiBoxBuilder.h"
+#include "SkeletalMesh/SReferenceSkeletonTree.h"
 
 #define LOCTEXT_NAMESPACE "SkinWeightToolSettingsEditor"
 
@@ -1310,7 +1311,7 @@ void FSkinWeightDetailCustomization::AddSelectionUI(IDetailLayoutBuilder& Detail
 			.ToolTipText(LOCTEXT("PruneButtonTooltip",
 				"Removes influences with weights below the given threshold value.\n"
 				"Pruned bones are removed from the list of bones affecting the given vertex.\n"
-				"Pruned bones will no longer recieve weight when a vertex is normalized.\n"
+				"Pruned bones will no longer receive weight when a vertex is normalized.\n"
 				"This command operates on the selected vertices."))
 			.OnClicked_Lambda([this]()
 			{
@@ -1735,6 +1736,7 @@ void SVertexWeightEditor::Construct(const FArguments& InArgs, USkinWeightsPaintT
 				return Tool.IsValid() && !Tool->GetMainMeshSelector()->GetSelectedVertices().IsEmpty() ? EVisibility::Visible : EVisibility::Collapsed;
 			})
 			.SelectionMode(ESelectionMode::Single)
+			.OnSelectionChanged(this, &SVertexWeightEditor::OnSelectionChanged)
 			.ListItemsSource( &ListViewItems )
 			.OnGenerateRow_Lambda([this](TSharedPtr<FWeightEditorElement> Element, const TSharedRef<STableViewBase>& OwnerTableView)
 			{
@@ -1784,4 +1786,22 @@ void SVertexWeightEditor::RefreshView()
 	ListView->RequestListRefresh();
 }
 
+void SVertexWeightEditor::OnSelectionChanged(TSharedPtr<FWeightEditorElement> Selection, ESelectInfo::Type SelectInfo)
+{
+	if (!Selection.IsValid() ||
+		!Tool.IsValid())
+	{
+		return;
+	}
+	const USkinWeightsPaintTool& ToolRef = *Tool.Get();
+	
+	if (const TSharedPtr<SReferenceSkeletonTree>& SkeletonTree = ToolRef.GetToolSkeletonTree())
+	{
+		const FName BoneName = ToolRef.GetBoneNameFromIndex(Selection->BoneIndex);
+
+		constexpr bool bFrameSelection = true;
+		SkeletonTree->SelectItemFromNames({ BoneName }, bFrameSelection);
+	}
+}
+
 #undef LOCTEXT_NAMESPACE
diff --git a/Engine/Plugins/Animation/SkeletalMeshModelingTools/Source/SkeletalMeshModelingTools/Private/DetailCustomization/SkinWeightDetailCustomization.h b/Engine/Plugins/Animation/SkeletalMeshModelingTools/Source/SkeletalMeshModelingTools/Private/DetailCustomization/SkinWeightDetailCustomization.h
--- a/Engine/Plugins/Animation/SkeletalMeshModelingTools/Source/SkeletalMeshModelingTools/Private/DetailCustomization/SkinWeightDetailCustomization.h
+++ b/Engine/Plugins/Animation/SkeletalMeshModelingTools/Source/SkeletalMeshModelingTools/Private/DetailCustomization/SkinWeightDetailCustomization.h
@@ -93,6 +93,8 @@ public:
 	void RefreshView();
 
 private:
+	/** Called when the selection in this editor changed */
+	void OnSelectionChanged(TSharedPtr<FWeightEditorElement> Selection, ESelectInfo::Type SelectInfo);
 
 	// the vertex weight list view
 	TSharedPtr<SWeightEditorListViewType> ListView;
diff --git a/Engine/Plugins/Runtime/MeshModelingToolset/Source/MeshModelingToolsEditorOnly/Private/SkeletalMesh/SkinWeightsPaintTool.cpp b/Engine/Plugins/Runtime/MeshModelingToolset/Source/MeshModelingToolsEditorOnly/Private/SkeletalMesh/SkinWeightsPaintTool.cpp
--- a/Engine/Plugins/Runtime/MeshModelingToolset/Source/MeshModelingToolsEditorOnly/Private/SkeletalMesh/SkinWeightsPaintTool.cpp
+++ b/Engine/Plugins/Runtime/MeshModelingToolset/Source/MeshModelingToolsEditorOnly/Private/SkeletalMesh/SkinWeightsPaintTool.cpp
@@ -4256,7 +4256,7 @@ TSharedPtr<SWidget> USkinWeightsPaintTool::GetCustomSkeletonTreeWidget(TSharedPt
 			return true;
 		});
 	
-	TSharedPtr<SReferenceSkeletonTree> ToolSkeletonTree =
+	ToolSkeletonTree =
 		SNew(SReferenceSkeletonTree)
 		.Modifier(SkeletonReader)
 		.Delegates(Delegates);
diff --git a/Engine/Plugins/Runtime/MeshModelingToolset/Source/MeshModelingToolsEditorOnly/Public/SkeletalMesh/SReferenceSkeletonTree.h b/Engine/Plugins/Runtime/MeshModelingToolset/Source/MeshModelingToolsEditorOnly/Public/SkeletalMesh/SReferenceSkeletonTree.h
--- a/Engine/Plugins/Runtime/MeshModelingToolset/Source/MeshModelingToolsEditorOnly/Public/SkeletalMesh/SReferenceSkeletonTree.h
+++ b/Engine/Plugins/Runtime/MeshModelingToolset/Source/MeshModelingToolsEditorOnly/Public/SkeletalMesh/SReferenceSkeletonTree.h
@@ -232,7 +232,7 @@ public:
 	UE_API void GetSelectedBoneNames(TArray<FName>& OutSelectedBoneNames) const;
 	TArray<TSharedPtr<FBoneElement>> GetSelectedItems() const;
 	bool HasSelectedItems() const;
-	void SelectItemFromNames(const TArray<FName>& InBoneNames, bool bFrameSelection = false);
+	UE_API void SelectItemFromNames(const TArray<FName>& InBoneNames, bool bFrameSelection = false);
 
 	UE_API TSharedPtr<ISkeletalMeshNotifier> GetNotifier();
 
diff --git a/Engine/Plugins/Runtime/MeshModelingToolset/Source/MeshModelingToolsEditorOnly/Public/SkeletalMesh/SkinWeightsPaintTool.h b/Engine/Plugins/Runtime/MeshModelingToolset/Source/MeshModelingToolsEditorOnly/Public/SkeletalMesh/SkinWeightsPaintTool.h
--- a/Engine/Plugins/Runtime/MeshModelingToolset/Source/MeshModelingToolsEditorOnly/Public/SkeletalMesh/SkinWeightsPaintTool.h
+++ b/Engine/Plugins/Runtime/MeshModelingToolset/Source/MeshModelingToolsEditorOnly/Public/SkeletalMesh/SkinWeightsPaintTool.h
@@ -36,6 +36,7 @@ class FEditorViewportClient;
 class UVolumetricBrushStampIndicator;
 class USkeletonModifier;
 class SButton;
+class SReferenceSkeletonTree;
 class FSkeletalMeshNotifierBindScope;
 
 namespace UE::Geometry 
@@ -826,6 +827,9 @@ public:
 	DECLARE_MULTICAST_DELEGATE(FOnWeightsChanged);
 	FOnWeightsChanged OnWeightsChanged;
 
+	/** Returns the skeleton tree displayed when using this tool */
+	const TSharedPtr<SReferenceSkeletonTree>& GetToolSkeletonTree() const { return ToolSkeletonTree; }
+
 protected:
 	// UInteractiveTool overrides
 	UE_API virtual void RegisterActions(FInteractiveToolActionSet& ActionSet) override;
@@ -961,6 +965,9 @@ protected:
 	UPROPERTY()
 	TObjectPtr<UWeightToolSelectionIsolator> SelectionIsolator = nullptr;
 	
+	/** The skeleton tree displayed when using this tool */
+	TSharedPtr<SReferenceSkeletonTree> ToolSkeletonTree;
+
 	// editor state to restore when exiting the paint tool
 	FString PreviewProfileToRestore;
 
