[SkeletalMeshModelingTools] Fixed alt left click drag not triggering camera orbit when paint skin weight tool is active

#Synced-CL 50430604

diff --git a/Engine/Plugins/Animation/SkeletalMeshModelingTools/Source/SkeletalMeshModelingTools/Private/SkeletalMeshModelingToolsEditorMode.cpp b/Engine/Plugins/Animation/SkeletalMeshModelingTools/Source/SkeletalMeshModelingTools/Private/SkeletalMeshModelingToolsEditorMode.cpp
--- a/Engine/Plugins/Animation/SkeletalMeshModelingTools/Source/SkeletalMeshModelingTools/Private/SkeletalMeshModelingToolsEditorMode.cpp
+++ b/Engine/Plugins/Animation/SkeletalMeshModelingTools/Source/SkeletalMeshModelingTools/Private/SkeletalMeshModelingToolsEditorMode.cpp
@@ -27,6 +27,7 @@
 #include "ViewportInteractions/ViewportInteractionsBehaviorSource.h"
 #include "ViewportInteractions/ViewportInteractionsUtilities.h"
 #include "ViewportInteractions/ViewportMoveYawInteraction.h"
+#include "ViewportInteractions/ViewportOrbitInteraction.h"
 #include "Widgets/Notifications/SNotificationList.h"
 
 #include UE_INLINE_GENERATED_CPP_BY_NAME(SkeletalMeshModelingToolsEditorMode)
@@ -396,6 +397,12 @@ void USkeletalMeshModelingToolsEditorMode::Enter()
 		{
 			UE::Editor::ViewportInteractions::AddDefaultCameraMovementInteractions(Source);
 			Source->AddInteraction<USkeletalMeshModelingToolsClickSelection>();
+
+			// Making sure viewport orbit is prioritized before regular tool drag tools	
+			for (const TObjectPtr<UInputBehavior>& InputBehavior : Source->FindInteraction(UViewportOrbitInteraction::StaticClass())->GetInputBehaviors())
+			{
+				InputBehavior->SetDefaultPriority(UE::Editor::ViewportInteractions::VIEWPORT_INTERACTIONS_DEFAULT_PRIORITY);
+			}
 				
 			check(WeakThis.IsValid());
 			USkeletalMeshModelingToolsGeometryFrustumSelection* GeometryFrustumSelection =
