[ITF Viewport Interactions]
- Add a `PostBuildViewportInteractions` event that is called with a const source after `OnBuildViewportInteractions`. This provides a known point where all interactions for a given mode set will be present and any post-processing (but not additions/removals) can occur.
- Add automatically-bound functions for the viewport interaction build events to `FAssetEditorToolkit`.
- Move the enabling of viewport interactions and the core viewport interactions to the top level of the Skeletal Mesh Editor.

#rnx

#Synced-CL 50708145

diff --git a/Engine/Plugins/Animation/SkeletalMeshModelingTools/Source/SkeletalMeshModelingTools/Private/SkeletalMeshModelingToolsEditorMode.cpp b/Engine/Plugins/Animation/SkeletalMeshModelingTools/Source/SkeletalMeshModelingTools/Private/SkeletalMeshModelingToolsEditorMode.cpp
--- a/Engine/Plugins/Animation/SkeletalMeshModelingTools/Source/SkeletalMeshModelingTools/Private/SkeletalMeshModelingToolsEditorMode.cpp
+++ b/Engine/Plugins/Animation/SkeletalMeshModelingTools/Source/SkeletalMeshModelingTools/Private/SkeletalMeshModelingToolsEditorMode.cpp
@@ -25,7 +25,6 @@
 #include "ViewportInteractions/SKMModelingToolsGeometryHoverInteraction.h"
 #include "ViewportInteractions/SKMModelingToolsGeometrySoftSelectRadiusDrag.h"
 #include "ViewportInteractions/ViewportInteractionsBehaviorSource.h"
-#include "ViewportInteractions/ViewportInteractionsUtilities.h"
 #include "ViewportInteractions/ViewportMoveYawInteraction.h"
 #include "ViewportInteractions/ViewportOrbitInteraction.h"
 #include "Widgets/Notifications/SNotificationList.h"
@@ -391,41 +390,8 @@ void USkeletalMeshModelingToolsEditorMode::Enter()
 
 	if (UModeManagerInteractiveToolsContext* const ModeManagerInteractiveToolsContext = GetModeManager()->GetInteractiveToolsContext())
 	{
-		ModeManagerInteractiveToolsContext->SetViewportInteractionsEnabled(true);
-
-		BuildViewportInteractionsDelegateHandle = ModeManagerInteractiveToolsContext->OnBuildViewportInteractions().AddLambda([WeakThis = TWeakObjectPtr<USkeletalMeshModelingToolsEditorMode>(this)](UViewportInteractionsBehaviorSource* Source)
-		{
-			UE::Editor::ViewportInteractions::AddDefaultCameraMovementInteractions(Source);
-			Source->AddInteraction<USkeletalMeshModelingToolsClickSelection>();
-
-			// Making sure viewport orbit is prioritized before regular tool drag tools	
-			for (const TObjectPtr<UInputBehavior>& InputBehavior : Source->FindInteraction(UViewportOrbitInteraction::StaticClass())->GetInputBehaviors())
-			{
-				InputBehavior->SetDefaultPriority(UE::Editor::ViewportInteractions::VIEWPORT_INTERACTIONS_DEFAULT_PRIORITY);
-			}
-				
-			check(WeakThis.IsValid());
-			USkeletalMeshModelingToolsGeometryFrustumSelection* GeometryFrustumSelection =
-				Source->AddInteraction<USkeletalMeshModelingToolsGeometryFrustumSelection>();
-			GeometryFrustumSelection->BindSelectionManager(WeakThis->SelectionManager);
-
-			USkeletalMeshModelingToolsGeometryClickSelection* GeometryClickSelection =
-				Source->AddInteraction<USkeletalMeshModelingToolsGeometryClickSelection>();
-			GeometryClickSelection->BindSelectionManager(WeakThis->SelectionManager);
-
-			USkeletalMeshModelingToolsGeometryHoverInteraction* GeometryHoverInteraction =
-				Source->AddInteraction<USkeletalMeshModelingToolsGeometryHoverInteraction>();
-			GeometryHoverInteraction->BindSelectionManager(WeakThis->SelectionManager);
-
-			USkeletalMeshModelingToolsGeometrySoftSelectRadiusDragInteraction* GeometrySoftSelectDrag =
-				Source->AddInteraction<USkeletalMeshModelingToolsGeometrySoftSelectRadiusDragInteraction>();
-			GeometrySoftSelectDrag->BindSelectionManager(WeakThis->SelectionManager);
-				
-				
-			// Since viewport interactions are rebuilt whenever an editor mode is changed (including modes that can coexists like the persona skeleton selection edit mode)
-			// we have to check again to make sure the correct interactions are enabled when that happens
-			WeakThis->ToggleGeometrySelectionViewportInteractions(WeakThis->IsGeometrySelectionActive());
-		});
+		ModeManagerInteractiveToolsContext->OnBuildViewportInteractions().AddUObject(this, &USkeletalMeshModelingToolsEditorMode::OnBuildViewportInteractions);
+		ModeManagerInteractiveToolsContext->PostBuildViewportInteractions().AddUObject(this, &USkeletalMeshModelingToolsEditorMode::PostBuildViewportInteractions);
 	}
 
 	
@@ -730,7 +696,8 @@ void USkeletalMeshModelingToolsEditorMode::Exit()
 {
 	if (UModeManagerInteractiveToolsContext* const ModeManagerInteractiveToolsContext = GetModeManager()->GetInteractiveToolsContext())
 	{
-		ModeManagerInteractiveToolsContext->OnBuildViewportInteractions().Remove(BuildViewportInteractionsDelegateHandle);
+		ModeManagerInteractiveToolsContext->OnBuildViewportInteractions().RemoveAll(this);
+		ModeManagerInteractiveToolsContext->PostBuildViewportInteractions().RemoveAll(this);
 	}
 	
 	RemoveViewportToolbarExtensions();
@@ -1202,6 +1169,34 @@ void USkeletalMeshModelingToolsEditorMode::RemoveViewportToolbarExtensions()
 	UToolMenus::Get()->UnregisterOwnerByName(SkeletalMeshModelingToolsEditorModeLocals::ViewportToolbarExtensionOwnerName);
 }
 
+void USkeletalMeshModelingToolsEditorMode::OnBuildViewportInteractions(UViewportInteractionsBehaviorSource* Source)
+{
+	Source->AddInteraction<USkeletalMeshModelingToolsClickSelection>();
+	
+	USkeletalMeshModelingToolsGeometryFrustumSelection* GeometryFrustumSelection =
+		Source->AddInteraction<USkeletalMeshModelingToolsGeometryFrustumSelection>();
+	GeometryFrustumSelection->BindSelectionManager(SelectionManager);
+
+	USkeletalMeshModelingToolsGeometryClickSelection* GeometryClickSelection =
+		Source->AddInteraction<USkeletalMeshModelingToolsGeometryClickSelection>();
+	GeometryClickSelection->BindSelectionManager(SelectionManager);
+
+	USkeletalMeshModelingToolsGeometryHoverInteraction* GeometryHoverInteraction =
+		Source->AddInteraction<USkeletalMeshModelingToolsGeometryHoverInteraction>();
+	GeometryHoverInteraction->BindSelectionManager(SelectionManager);
+
+	USkeletalMeshModelingToolsGeometrySoftSelectRadiusDragInteraction* GeometrySoftSelectDrag =
+		Source->AddInteraction<USkeletalMeshModelingToolsGeometrySoftSelectRadiusDragInteraction>();
+	GeometrySoftSelectDrag->BindSelectionManager(SelectionManager);
+}
+
+void USkeletalMeshModelingToolsEditorMode::PostBuildViewportInteractions(const UViewportInteractionsBehaviorSource* Source)
+{
+	// Since viewport interactions are rebuilt whenever an editor mode is changed (including modes that can coexists like the persona skeleton selection edit mode)
+	// we have to check again to make sure the correct interactions are enabled when that happens
+	ToggleGeometrySelectionViewportInteractions(IsGeometrySelectionActive());
+}
+
 bool USkeletalMeshModelingToolsEditorMode::IsGeometrySelectionActive() const
 {
 	return bAllowGeometrySelection && SelectionManager->GetMeshTopologyMode() != UGeometrySelectionManager::EMeshTopologyMode::None;
diff --git a/Engine/Plugins/Animation/SkeletalMeshModelingTools/Source/SkeletalMeshModelingTools/Private/SkeletalMeshModelingToolsEditorMode.h b/Engine/Plugins/Animation/SkeletalMeshModelingTools/Source/SkeletalMeshModelingTools/Private/SkeletalMeshModelingToolsEditorMode.h
--- a/Engine/Plugins/Animation/SkeletalMeshModelingTools/Source/SkeletalMeshModelingTools/Private/SkeletalMeshModelingToolsEditorMode.h
+++ b/Engine/Plugins/Animation/SkeletalMeshModelingTools/Source/SkeletalMeshModelingTools/Private/SkeletalMeshModelingToolsEditorMode.h
@@ -12,9 +12,11 @@
 #include "Interfaces/ISKMBackedDynaMeshComponentProvider.h"
 #include "Changes/ValueWatcher.h"
 #include "TargetInterfaces/MeshTargetInterfaceTypes.h"
+#include "ViewportInteractions/ViewportInteractionsBehaviorSource.h"
 
 #include "SkeletalMeshModelingToolsEditorMode.generated.h"
 
+class UViewportInteractionsBehaviorSource;
 class USkeletalMeshGeometrySelectionTransformTweaker;
 class FSkeletalMeshModelingToolsEditorModeToolkit;
 
@@ -282,12 +284,13 @@ private:
 	void PopulateViewportToolbar();
 	void RemoveViewportToolbarExtensions();
 
+	void OnBuildViewportInteractions(UViewportInteractionsBehaviorSource* Source);
+	void PostBuildViewportInteractions(const UViewportInteractionsBehaviorSource* Source);
+
 	// Used to keep the command objects alive
 	UPROPERTY(NonTransactional)
 	TArray<TObjectPtr<UInteractiveCommand>> ModelingModeCommands;
 
-	FDelegateHandle BuildViewportInteractionsDelegateHandle;
-
 	TValueWatcher<bool> GeometrySelectionStateWatcher;
 
 	bool IsGeometrySelectionActive() const;
diff --git a/Engine/Source/Editor/SkeletalMeshEditor/Private/SkeletalMeshEditor.cpp b/Engine/Source/Editor/SkeletalMeshEditor/Private/SkeletalMeshEditor.cpp
--- a/Engine/Source/Editor/SkeletalMeshEditor/Private/SkeletalMeshEditor.cpp
+++ b/Engine/Source/Editor/SkeletalMeshEditor/Private/SkeletalMeshEditor.cpp
@@ -68,6 +68,10 @@
 #include "InterchangeAssetImportData.h"
 #include "InterchangeManager.h"
 #include "SkeletalMeshEditorModeUILayer.h"
+#include "Tools/EdModeInteractiveToolsContext.h"
+#include "ViewportInteractions/ViewportInteractionsBehaviorSource.h"
+#include "ViewportInteractions/ViewportInteractionsUtilities.h"
+#include "ViewportInteractions/ViewportOrbitInteraction.h"
 
 const FName SkeletalMeshEditorAppIdentifier = FName(TEXT("SkeletalMeshEditorApp"));
 
@@ -271,6 +275,11 @@ void FSkeletalMeshEditor::InitSkeletalMeshEditor(const EToolkitMode::Type Mode,
 	const bool bCreateDefaultToolbar = true;
 	FAssetEditorToolkit::InitAssetEditor(Mode, InitToolkitHost, SkeletalMeshEditorAppIdentifier, FTabManager::FLayout::NullLayout, bCreateDefaultStandaloneMenu, bCreateDefaultToolbar, InSkeletalMesh);
 	
+	if (UModeManagerInteractiveToolsContext* ToolsContext = GetEditorModeManager().GetInteractiveToolsContext())
+	{
+		ToolsContext->SetViewportInteractionsEnabled(true);
+	}
+	
 	TSharedPtr<class IToolkitHost> PinnedToolkitHost = ToolkitHost.Pin();
 	check(PinnedToolkitHost.IsValid());
 	ModeUILayer = MakeShareable(new FSkeletalMeshEditorModeUILayer(PinnedToolkitHost.Get()));
@@ -1619,6 +1628,17 @@ TSharedPtr<FUICommandInfo> FSkeletalMeshEditor::GetResetAllBonesTransformsComman
 	return Commands.ResetAllBonesTransforms;
 }
 
+void FSkeletalMeshEditor::OnBuildViewportInteractions(UViewportInteractionsBehaviorSource* Source)
+{
+	UE::Editor::ViewportInteractions::AddDefaultCameraMovementInteractions(Source);
+			
+	// Making sure viewport orbit is prioritized before regular tool drag tools	
+	for (const TObjectPtr<UInputBehavior>& InputBehavior : Source->FindInteraction(UViewportOrbitInteraction::StaticClass())->GetInputBehaviors())
+	{
+		InputBehavior->SetDefaultPriority(UE::Editor::ViewportInteractions::VIEWPORT_INTERACTIONS_DEFAULT_PRIORITY);
+	}
+}
+
 FSkeletalMeshEditorBinding::FSkeletalMeshEditorBinding(TSharedRef<FSkeletalMeshEditor> InEditor)
 	: Editor(InEditor)
 {}
diff --git a/Engine/Source/Editor/SkeletalMeshEditor/Private/SkeletalMeshEditor.h b/Engine/Source/Editor/SkeletalMeshEditor/Private/SkeletalMeshEditor.h
--- a/Engine/Source/Editor/SkeletalMeshEditor/Private/SkeletalMeshEditor.h
+++ b/Engine/Source/Editor/SkeletalMeshEditor/Private/SkeletalMeshEditor.h
@@ -136,6 +136,9 @@ public:
 	virtual FSimpleMulticastDelegate& OnPreSaveAssetAs() override;
 	virtual TSharedPtr<FUICommandInfo> GetResetBoneTransformsCommand() override;
 	virtual TSharedPtr<FUICommandInfo> GetResetAllBonesTransformsCommand() override;
+	
+protected:
+	void OnBuildViewportInteractions(UViewportInteractionsBehaviorSource* Source) override;
 
 private:
 	void HandleObjectSelected(UObject* InObject);
diff --git a/Engine/Source/Editor/SkeletalMeshEditor/SkeletalMeshEditor.Build.cs b/Engine/Source/Editor/SkeletalMeshEditor/SkeletalMeshEditor.Build.cs
--- a/Engine/Source/Editor/SkeletalMeshEditor/SkeletalMeshEditor.Build.cs
+++ b/Engine/Source/Editor/SkeletalMeshEditor/SkeletalMeshEditor.Build.cs
@@ -22,6 +22,7 @@ public class SkeletalMeshEditor : ModuleRules
 				"Slate",
 				"SlateCore",
 				"EditorFramework",
+				"EditorInteractiveToolsFramework",
                 "UnrealEd",
                 "SkeletonEditor",
                 "Kismet",
diff --git a/Engine/Source/Editor/UnrealEd/Private/Toolkits/AssetEditorToolkit.cpp b/Engine/Source/Editor/UnrealEd/Private/Toolkits/AssetEditorToolkit.cpp
--- a/Engine/Source/Editor/UnrealEd/Private/Toolkits/AssetEditorToolkit.cpp
+++ b/Engine/Source/Editor/UnrealEd/Private/Toolkits/AssetEditorToolkit.cpp
@@ -39,6 +39,7 @@
 #include "Toolkits/GlobalEditorCommonCommands.h"
 #include "Toolkits/SStandaloneAssetEditorToolkitHost.h"
 #include "Toolkits/ToolkitManager.h"
+#include "Tools/EdModeInteractiveToolsContext.h"
 #include "Widgets/SToolTip.h"
 #include "Widgets/Docking/SDockTab.h"
 #include "Widgets/Images/SImage.h"
@@ -319,6 +320,12 @@ void FAssetEditorToolkit::InitAssetEditor( const EToolkitMode::Type Mode, const
 	if (EditorModeManager)
 	{
 		EditorModeManager->SetToolkitHost(ToolkitHost.Pin().ToSharedRef());
+		
+		if (UModeManagerInteractiveToolsContext* ToolsContext = EditorModeManager->GetInteractiveToolsContext())
+		{
+			ToolsContext->OnBuildViewportInteractions().AddRaw(this, &FAssetEditorToolkit::OnBuildViewportInteractions);
+			ToolsContext->PostBuildViewportInteractions().AddRaw(this, &FAssetEditorToolkit::PostBuildViewportInteractions);
+		}
 	}
 	
 	PostInitAssetEditor();
@@ -340,6 +347,15 @@ FAssetEditorToolkit::~FAssetEditorToolkit()
 			AssetEditorSubsystem->NotifyEditorClosed(this);
 		}
 	}
+	
+	if (EditorModeManager)
+	{
+		if (UModeManagerInteractiveToolsContext* ToolsContext = EditorModeManager->GetInteractiveToolsContext())
+		{
+			ToolsContext->OnBuildViewportInteractions().RemoveAll(this);
+			ToolsContext->PostBuildViewportInteractions().RemoveAll(this);
+		}
+	}
 
 	EditorModeManager.Reset();
 
diff --git a/Engine/Source/Editor/UnrealEd/Private/Tools/EdModeInteractiveToolsContext.cpp b/Engine/Source/Editor/UnrealEd/Private/Tools/EdModeInteractiveToolsContext.cpp
--- a/Engine/Source/Editor/UnrealEd/Private/Tools/EdModeInteractiveToolsContext.cpp
+++ b/Engine/Source/Editor/UnrealEd/Private/Tools/EdModeInteractiveToolsContext.cpp
@@ -1576,6 +1576,7 @@ void UModeManagerInteractiveToolsContext::RegisterViewportInteractions()
 	OnGetUnsupportedInteractions().Broadcast(UnsupportedInteractions);
 	ViewportInteractionsBehaviorSource->SetUnsupportedViewportInteractions(UnsupportedInteractions);
 	OnBuildViewportInteractions().Broadcast(ViewportInteractionsBehaviorSource);
+	PostBuildViewportInteractions().Broadcast(ViewportInteractionsBehaviorSource);
 
 	ViewportInteractionsBehaviorSource->RegisterBehaviorSources();
 }
diff --git a/Engine/Source/Editor/UnrealEd/Public/Toolkits/AssetEditorToolkit.h b/Engine/Source/Editor/UnrealEd/Public/Toolkits/AssetEditorToolkit.h
--- a/Engine/Source/Editor/UnrealEd/Public/Toolkits/AssetEditorToolkit.h
+++ b/Engine/Source/Editor/UnrealEd/Public/Toolkits/AssetEditorToolkit.h
@@ -26,6 +26,7 @@ class SBorder;
 class SStandaloneAssetEditorToolkitHost;
 class FWorkspaceItem;
 class UToolMenu;
+class UViewportInteractionsBehaviorSource;
 
 namespace UE::Editor::Toolbars
 {
@@ -483,6 +484,12 @@ protected:
 	 * If it's not overridden or NAME_None is returned, the default id will be used.
 	 */
 	virtual FName GetOverrideTabId() const { return NAME_None; }
+	
+	UE_EXPERIMENTAL(5.8, "Viewport Interactions are Experimental")
+	virtual void OnBuildViewportInteractions(UViewportInteractionsBehaviorSource* Source) {};
+	
+	UE_EXPERIMENTAL(5.8, "Viewport Interactions are Experimental")
+	virtual void PostBuildViewportInteractions(const UViewportInteractionsBehaviorSource* Source) {};
 
 private:
 	// Callback for persisting the Asset Editor's layout.
diff --git a/Engine/Source/Editor/UnrealEd/Public/Tools/EdModeInteractiveToolsContext.h b/Engine/Source/Editor/UnrealEd/Public/Tools/EdModeInteractiveToolsContext.h
--- a/Engine/Source/Editor/UnrealEd/Public/Tools/EdModeInteractiveToolsContext.h
+++ b/Engine/Source/Editor/UnrealEd/Public/Tools/EdModeInteractiveToolsContext.h
@@ -424,17 +424,27 @@ PRAGMA_DISABLE_DEPRECATION_WARNINGS
 	UE_DEPRECATED(5.8, "Use OnBuildViewportInteractions instead.")
 	FViewportInteractionsRegisterStateChanged OnRegisterViewportInteractionsDelegate;
 PRAGMA_ENABLE_DEPRECATION_WARNINGS
-	
+
+PRAGMA_DISABLE_EXPERIMENTAL_WARNINGS	
 	UE_EXPERIMENTAL(5.8, "Viewport interactions are experimental")
 	DECLARE_MULTICAST_DELEGATE_OneParam(FViewportInteractionsSourceMessage, UViewportInteractionsBehaviorSource*)
 	
+	UE_EXPERIMENTAL(5.8, "Viewport interactions are experimental")
+	DECLARE_MULTICAST_DELEGATE_OneParam(FViewportInteractionsSourceConstMessage, const UViewportInteractionsBehaviorSource*)
+	
 	UE_EXPERIMENTAL(5.8, "Viewport interactions are experimental")
 	/**
 	 * Called when viewport interactions are built. Add any custom interactions to the behavior source that is passed in via the event.
 	 * Interactions are automatically removed when editor modes change.
 	 */
-PRAGMA_DISABLE_EXPERIMENTAL_WARNINGS
 	FViewportInteractionsSourceMessage& OnBuildViewportInteractions() { return OnBuildViewportInteractionsEvent; }
+
+
+	UE_EXPERIMENTAL(5.8, "Viewport interactions are experimental")
+	/**
+	 * Called after viewport interactions are built. Can be used to apply any post-processing of behaviors after all behaviors are present.
+	 */
+	FViewportInteractionsSourceConstMessage& PostBuildViewportInteractions() { return PostBuildViewportInteractionsEvent; }
 PRAGMA_ENABLE_EXPERIMENTAL_WARNINGS
 
 protected:
@@ -489,6 +499,9 @@ protected:
 	PRAGMA_DISABLE_EXPERIMENTAL_WARNINGS
 	UE_EXPERIMENTAL(5.8, "Viewport interactions are experimental")
 	FViewportInteractionsSourceMessage OnBuildViewportInteractionsEvent;
+	
+	UE_EXPERIMENTAL(5.8, "Viewport interactions are experimental")
+	FViewportInteractionsSourceConstMessage PostBuildViewportInteractionsEvent;
 	PRAGMA_ENABLE_EXPERIMENTAL_WARNINGS
 
 	bool bUsesViewportInteractions = false;
