[SkeletalMeshModelingTools] Added menu option to mirror a morph target

#Synced-CL 50896465

diff --git a/Engine/Plugins/Animation/SkeletalMeshModelingTools/Source/SkeletalMeshModelingTools/Private/Components/SKMBackedDynaMeshComponent.cpp b/Engine/Plugins/Animation/SkeletalMeshModelingTools/Source/SkeletalMeshModelingTools/Private/Components/SKMBackedDynaMeshComponent.cpp
--- a/Engine/Plugins/Animation/SkeletalMeshModelingTools/Source/SkeletalMeshModelingTools/Private/Components/SKMBackedDynaMeshComponent.cpp
+++ b/Engine/Plugins/Animation/SkeletalMeshModelingTools/Source/SkeletalMeshModelingTools/Private/Components/SKMBackedDynaMeshComponent.cpp
@@ -6,12 +6,12 @@
 #include "LODUtilities.h"
 #include "MeshDescriptionToDynamicMesh.h"
 #include "ModelingToolTargetEditorOnlyUtil.h"
-#include "ModelingToolTargetUtil.h"
 #include "SkeletalMeshAttributes.h"
 #include "SkeletonModifier.h"
 #include "ScopedTransaction.h"
 #include "Misc/ITransaction.h"
 #include "SkeletalMeshEditorSubsystem.h"
+#include "Parameterization/MeshPlanarSymmetry.h"
 #include "ToolTargets/SkeletalMeshToolTarget.h"
 #include "SkeletalMesh/SkeletalMeshToolsHelper.h"
 
@@ -204,6 +204,7 @@ EMeshLODIdentifier USkeletalMeshBackedDynamicMeshComponent::Init(USkeletalMesh*
 
 	SetMesh(MoveTemp(LODMesh));
 
+	UpdateSymmetry();
 	// Additional Trackers
 	ResetTrackersDirect();
 
@@ -433,6 +434,10 @@ void USkeletalMeshBackedDynamicMeshComponent::HandleGeometryUpdate(
 	{
 		Trackers.MorphTargetChangeTracker.HandleEditMorphTarget(EditingMorphTargetName);
 	}
+	else
+	{
+		MarkBaseGeometryDirtyInternal();
+	}
 	
 	MarkDirtyInternal();
 }
@@ -561,8 +566,8 @@ TArray<FName> USkeletalMeshBackedDynamicMeshComponent::DuplicateMorphTargets(con
 			FDynamicMeshMorphTargetAttribute* SourceMorphTargetAttribute = Mesh.Attributes()->GetMorphTargetAttribute(Name);
 			
 			FDynamicMeshMorphTargetAttribute* DuplicatedMorphTargetAttribute = new FDynamicMeshMorphTargetAttribute(&Mesh);
+			DuplicatedMorphTargetAttribute->Copy(*SourceMorphTargetAttribute);
 			DuplicatedMorphTargetAttribute->SetName(DuplicateName);
-			*DuplicatedMorphTargetAttribute = (*SourceMorphTargetAttribute);
 			
 			Mesh.Attributes()->AttachMorphTargetAttribute(DuplicateName, DuplicatedMorphTargetAttribute);
 		}
@@ -579,6 +584,85 @@ TArray<FName> USkeletalMeshBackedDynamicMeshComponent::DuplicateMorphTargets(con
 	return NewMorphs;
 }
 
+void USkeletalMeshBackedDynamicMeshComponent::MirrorMorphTargets(const TArray<FName>& InNames)
+{
+	if (!ReadOnlySymmetry)
+	{
+		return;
+	}
+	
+	TArray<FName> MorphsToEdit;
+	for (const FName& Name : InNames)
+	{
+		if (GetMorphTargetChangeTracker().GetCurrentMorphTargetNames().Contains(Name))
+		{
+			MorphsToEdit.Add(Name);
+		}
+	}
+	
+	if (MorphsToEdit.IsEmpty())
+	{
+		return;
+	}
+
+	FChangeScope ChangeScope(this);
+
+	EditMesh([&](FDynamicMesh3& Mesh)
+	{
+		using namespace UE::Geometry;
+		for (const FName& Name : MorphsToEdit)
+		{
+			using namespace UE::Geometry;
+			FDynamicMeshMorphTargetAttribute* SourceMorphTargetAttribute = Mesh.Attributes()->GetMorphTargetAttribute(Name);
+
+			ParallelFor(Mesh.MaxVertexID(), [&](int32 VertID)
+			{
+				if (!Mesh.IsVertex(VertID))
+				{
+					return;
+				}
+
+				FVector Position = Mesh.GetVertex(VertID);
+					
+				FVector Delta;
+				SourceMorphTargetAttribute->GetValue(VertID, Delta);
+
+				// Making sure on-plane verts are dealt with first
+				TArray<FVector> ConstrainedPosition;
+				ConstrainedPosition.Add(Position+Delta);
+				ReadOnlySymmetry->ApplySymmetryPlaneConstraints({VertID}, ConstrainedPosition);
+
+				Delta = ConstrainedPosition[0] - Position;
+				SourceMorphTargetAttribute->SetValue(VertID, Delta);
+
+				// Process positive side only	
+				if (Position.X < 0)
+				{
+					return;
+				}
+
+				FVector MirroredDelta = ReadOnlySymmetry->GetMirroredAxis(Delta);
+				
+				TArray<int32> MirroredVertID;
+				ReadOnlySymmetry->GetMirrorVertexROI({VertID}, MirroredVertID, true);
+
+				if (MirroredVertID[0] != INDEX_NONE)
+				{
+					SourceMorphTargetAttribute->SetValue(MirroredVertID[0], MirroredDelta);
+				}
+			});
+		}
+	});
+	
+	// Tracker change
+	for (const FName& Name : MorphsToEdit)
+	{
+		Trackers.MorphTargetChangeTracker.HandleEditMorphTarget(Name);
+	}
+
+	MarkDirtyInternal();	
+}
+
 void USkeletalMeshBackedDynamicMeshComponent::MarkMorphTargetEdited(FName InName)
 {
 	FTrackerChangeScope Scope(this);
@@ -613,6 +697,7 @@ void USkeletalMeshBackedDynamicMeshComponent::ResetTrackersDirect()
 	Trackers.InitialAssetMesh = GetMeshCopy();
 
 	Trackers.ChangeCount = 0;
+	Trackers.BaseGeometryChangeCount = 0;
 	
 	Trackers.RefSkeleton = WeakSkeletalMesh->GetRefSkeleton();
 	Trackers.RefSkeleton.GetBoneAbsoluteTransforms(Trackers.ComponentSpaceBoneTransformsRefPose);
@@ -663,6 +748,11 @@ void USkeletalMeshBackedDynamicMeshComponent::MarkDirtyInternal()
 	Trackers.ChangeCount++;
 }
 
+void USkeletalMeshBackedDynamicMeshComponent::MarkBaseGeometryDirtyInternal()
+{
+	Trackers.BaseGeometryChangeCount++;
+	UpdateSymmetry();
+}
 
 void USkeletalMeshBackedDynamicMeshComponent::MarkDirty()
 {
@@ -671,6 +761,12 @@ void USkeletalMeshBackedDynamicMeshComponent::MarkDirty()
 	MarkDirtyInternal();
 }
 
+void USkeletalMeshBackedDynamicMeshComponent::MarkBaseGeometryDirty()
+{
+	FTrackerChangeScope ChangeScope(this);
+
+	MarkBaseGeometryDirtyInternal();	
+}
 
 
 USkeletalMeshBackedDynamicMeshComponent::FTrackerChange::FTrackerChange(USkeletalMeshBackedDynamicMeshComponent* InComponent)
@@ -688,6 +784,10 @@ void USkeletalMeshBackedDynamicMeshComponent::FTrackerChange::Apply(UObject* Obj
 {
 	if (USkeletalMeshBackedDynamicMeshComponent* Component = Cast<USkeletalMeshBackedDynamicMeshComponent>(Object))
 	{
+		if (Component->Trackers.BaseGeometryChangeCount != NewTrackers.BaseGeometryChangeCount)
+		{
+			Component->UpdateSymmetry();
+		}
 		Component->Trackers = NewTrackers;
 	}	
 }
@@ -696,6 +796,10 @@ void USkeletalMeshBackedDynamicMeshComponent::FTrackerChange::Revert(UObject* Ob
 {
 	if (USkeletalMeshBackedDynamicMeshComponent* Component = Cast<USkeletalMeshBackedDynamicMeshComponent>(Object))
 	{
+		if (Component->Trackers.BaseGeometryChangeCount != OldTrackers.BaseGeometryChangeCount)
+		{
+			Component->UpdateSymmetry();
+		}
 		Component->Trackers = OldTrackers;
 	}	
 }
@@ -743,4 +847,22 @@ USkeletalMeshBackedDynamicMeshComponent::FChangeScope::~FChangeScope()
 	}
 }
 
+void USkeletalMeshBackedDynamicMeshComponent::UpdateSymmetry()
+{
+	using namespace UE::Geometry;
+
+
+	GetDynamicMesh()->ProcessMesh([this](const FDynamicMesh3& Mesh)
+	{
+		FAxisAlignedBox3d Bounds = Mesh.GetBounds(true);
+		TArray<FVector3d> PreferSymmetryPlaneNormal;
+		PreferSymmetryPlaneNormal.Add(FVector3d::UnitX());
+
+		FFrame3d SymmetryFrame(FVector3d::Zero(), FVector3d::UnitX());		
+			
+		ReadOnlySymmetry = MakePimpl<FMeshPlanarSymmetry>();
+		ReadOnlySymmetry->Initialize(const_cast<FDynamicMesh3*>(&Mesh), Bounds, SymmetryFrame);
+	});
+}
+
 #undef LOCTEXT_NAMESPACE
diff --git a/Engine/Plugins/Animation/SkeletalMeshModelingTools/Source/SkeletalMeshModelingTools/Private/Components/SKMBackedDynaMeshComponent.h b/Engine/Plugins/Animation/SkeletalMeshModelingTools/Source/SkeletalMeshModelingTools/Private/Components/SKMBackedDynaMeshComponent.h
--- a/Engine/Plugins/Animation/SkeletalMeshModelingTools/Source/SkeletalMeshModelingTools/Private/Components/SKMBackedDynaMeshComponent.h
+++ b/Engine/Plugins/Animation/SkeletalMeshModelingTools/Source/SkeletalMeshModelingTools/Private/Components/SKMBackedDynaMeshComponent.h
@@ -12,6 +12,11 @@
 #define UE_API SKELETALMESHMODELINGTOOLS_API
 
 
+namespace UE::Geometry
+{
+	class FMeshPlanarSymmetry;
+}
+
 class USkeletonModifier;
 
 
@@ -72,6 +77,7 @@ public:
 	void DiscardChanges();
 
 	void MarkDirty();
+	void MarkBaseGeometryDirty();
 	void HandleSkeletonChange(USkeletonModifier* InModifier);
 	void ForwardVisibilityChangeRequest(bool bInVisible);
 
@@ -86,6 +92,7 @@ public:
 	FName RenameMorphTarget(FName InOldName, FName InNewName);
 	void RemoveMorphTargets(const TArray<FName>& InNames);
 	TArray<FName> DuplicateMorphTargets(const TArray<FName>& InNames);
+	void MirrorMorphTargets(const TArray<FName>& InNames);
 	
 	void MarkMorphTargetEdited(FName InName);
 	
@@ -111,12 +118,15 @@ protected:
 
 	TSharedPtr<FDynamicMesh3> GetMeshCopy();
 	void MarkDirtyInternal();
+	void MarkBaseGeometryDirtyInternal();
 
 	struct FTrackers
 	{
 		TSharedPtr<FDynamicMesh3> InitialAssetMesh;
 		
 		int32 ChangeCount = 0;
+		// Base Geometry determines if symmetry tracking needs to be updated
+		int32 BaseGeometryChangeCount = 0;	
 		
 		FReferenceSkeleton RefSkeleton;
 		TArray<FTransform> ComponentSpaceBoneTransformsRefPose;
@@ -174,6 +184,10 @@ protected:
 	
 	friend class FChangeScope;
 	friend class FTrackerChange;
+
+	// Query only symmetry, not allow to modify mesh
+	TPimplPtr<UE::Geometry::FMeshPlanarSymmetry> ReadOnlySymmetry;
+	void UpdateSymmetry();
 };
 
 #undef UE_API
diff --git a/Engine/Plugins/Animation/SkeletalMeshModelingTools/Source/SkeletalMeshModelingTools/Private/SMorphTargetManager.cpp b/Engine/Plugins/Animation/SkeletalMeshModelingTools/Source/SkeletalMeshModelingTools/Private/SMorphTargetManager.cpp
--- a/Engine/Plugins/Animation/SkeletalMeshModelingTools/Source/SkeletalMeshModelingTools/Private/SMorphTargetManager.cpp
+++ b/Engine/Plugins/Animation/SkeletalMeshModelingTools/Source/SkeletalMeshModelingTools/Private/SMorphTargetManager.cpp
@@ -353,6 +353,11 @@ void SMorphTargetManager::BindCommands()
 		GenericCommands.Duplicate,
 		FExecuteAction::CreateSP(this, &SMorphTargetManager::DuplicateSelectedMorphTargets), 
 		FCanExecuteAction::CreateSP(this, &SMorphTargetManager::CanDuplicate));
+	
+	CommandList->MapAction(
+		Commands.MirrorMorphTarget,
+		FExecuteAction::CreateSP(this, &SMorphTargetManager::MirrorSelectedMorphTarget),
+		FCanExecuteAction::CreateSP(this, &SMorphTargetManager::CanMirror));
 }
 
 void SMorphTargetManager::RefreshList()
@@ -468,7 +473,8 @@ TSharedPtr<SWidget> SMorphTargetManager::OnGetContextMenuContent()
 	MenuBuilder.AddMenuEntry(FGenericCommands::Get().Rename);
 	MenuBuilder.AddMenuEntry(FGenericCommands::Get().Delete);
 	MenuBuilder.AddMenuEntry(FGenericCommands::Get().Duplicate);
-
+	MenuBuilder.AddMenuEntry(FSkeletalMeshModelingToolsCommands::Get().MirrorMorphTarget);
+	
 	MenuBuilder.EndSection();
 	return MenuBuilder.MakeWidget();
 }
@@ -578,6 +584,25 @@ bool SMorphTargetManager::CanDuplicate()
 	return ListView->GetSelectedItems().Num() > 0;
 }
 
+void SMorphTargetManager::MirrorSelectedMorphTarget()
+{
+	TArray<FMorphTargetInfoPtr> Items = ListView->GetSelectedItems();
+
+	check(Items.Num() == 1);
+
+	FName MorphTargetName = Items[0]->Name;
+	FText Title = FText::Format( LOCTEXT("MirrorMorphTransaction", "Mirror(Symmetrize) {0}"), FText::FromName(MorphTargetName));
+	FScopedTransaction Transaction(Title);
+
+	SetEditingMorphTarget(MorphTargetName);
+	Delegates.OnMirrorSelectedMorphTarget.Execute(MorphTargetName);
+}
+
+bool SMorphTargetManager::CanMirror()
+{
+	return ListView->GetSelectedItems().Num() == 1;
+}
+
 
 TSharedRef<class ITableRow> SMorphTargetManager::GenerateMorphTargetRow(
 	FMorphTargetInfoPtr MorphTargetItem,
diff --git a/Engine/Plugins/Animation/SkeletalMeshModelingTools/Source/SkeletalMeshModelingTools/Private/SMorphTargetManager.h b/Engine/Plugins/Animation/SkeletalMeshModelingTools/Source/SkeletalMeshModelingTools/Private/SMorphTargetManager.h
--- a/Engine/Plugins/Animation/SkeletalMeshModelingTools/Source/SkeletalMeshModelingTools/Private/SMorphTargetManager.h
+++ b/Engine/Plugins/Animation/SkeletalMeshModelingTools/Source/SkeletalMeshModelingTools/Private/SMorphTargetManager.h
@@ -36,6 +36,7 @@ public:
 	DECLARE_DELEGATE_RetVal_TwoParams(FName, FOnRenameMorphTarget, FName, FName);
 	DECLARE_DELEGATE_OneParam(FOnRemoveMorphTargets, const TArray<FName>&);
 	DECLARE_DELEGATE_RetVal_OneParam(TArray<FName>, FOnDuplicateMorphTargets, const TArray<FName>&);
+	DECLARE_DELEGATE_OneParam(FOnMirrorMorphTarget, FName);
 
 	struct FMorphTargetManagerDelegates
 	{
@@ -54,6 +55,7 @@ public:
 		FOnRenameMorphTarget OnRenameMorphTarget;
 		FOnRemoveMorphTargets OnRemoveMorphTargets;
 		FOnDuplicateMorphTargets OnDuplicateMorphTargets;
+		FOnMirrorMorphTarget OnMirrorSelectedMorphTarget;
 	};	
 	
 	SLATE_BEGIN_ARGS( SMorphTargetManager ) {}
@@ -99,6 +101,8 @@ protected:
 	bool CanRemove();
 	void DuplicateSelectedMorphTargets();
 	bool CanDuplicate();
+	void MirrorSelectedMorphTarget();
+	bool CanMirror();
 	
 	TSharedPtr<SSearchBox> NameFilterBox;
 	TSharedPtr<SMorphTargetManagerListType> ListView;
diff --git a/Engine/Plugins/Animation/SkeletalMeshModelingTools/Source/SkeletalMeshModelingTools/Private/SkeletalMeshEditingCache.cpp b/Engine/Plugins/Animation/SkeletalMeshModelingTools/Source/SkeletalMeshModelingTools/Private/SkeletalMeshEditingCache.cpp
--- a/Engine/Plugins/Animation/SkeletalMeshModelingTools/Source/SkeletalMeshModelingTools/Private/SkeletalMeshEditingCache.cpp
+++ b/Engine/Plugins/Animation/SkeletalMeshModelingTools/Source/SkeletalMeshModelingTools/Private/SkeletalMeshEditingCache.cpp
@@ -745,6 +745,11 @@ TArray<FName> USkeletalMeshEditingCache::DuplicateMorphTargets(const TArray<FNam
 	return EditingMeshComponent->DuplicateMorphTargets(InNames);
 }
 
+void USkeletalMeshEditingCache::MirrorMorphTargets(const TArray<FName>& InNames)
+{
+	EditingMeshComponent->MirrorMorphTargets(InNames);
+}
+
 void USkeletalMeshEditingCache::HandleGeometryUpdate(const FDynamicMesh3& InMesh, FName MorphTargetName)
 {
 	using namespace SkeletalMeshToolsHelper;
diff --git a/Engine/Plugins/Animation/SkeletalMeshModelingTools/Source/SkeletalMeshModelingTools/Private/SkeletalMeshEditingCache.h b/Engine/Plugins/Animation/SkeletalMeshModelingTools/Source/SkeletalMeshModelingTools/Private/SkeletalMeshEditingCache.h
--- a/Engine/Plugins/Animation/SkeletalMeshModelingTools/Source/SkeletalMeshModelingTools/Private/SkeletalMeshEditingCache.h
+++ b/Engine/Plugins/Animation/SkeletalMeshModelingTools/Source/SkeletalMeshModelingTools/Private/SkeletalMeshEditingCache.h
@@ -101,6 +101,7 @@ public:
 	FName RenameMorphTarget(FName InOldName, FName InNewName);
 	void RemoveMorphTargets(const TArray<FName>& InNames);
 	TArray<FName> DuplicateMorphTargets(const TArray<FName>& InNames);
+	void MirrorMorphTargets(const TArray<FName>& InNames);
 
 	void HandleGeometryUpdate(const UE::Geometry::FDynamicMesh3& InMesh, FName MorphTargetName);
 	
diff --git a/Engine/Plugins/Animation/SkeletalMeshModelingTools/Source/SkeletalMeshModelingTools/Private/SkeletalMeshModelingToolsCommands.cpp b/Engine/Plugins/Animation/SkeletalMeshModelingTools/Source/SkeletalMeshModelingTools/Private/SkeletalMeshModelingToolsCommands.cpp
--- a/Engine/Plugins/Animation/SkeletalMeshModelingTools/Source/SkeletalMeshModelingTools/Private/SkeletalMeshModelingToolsCommands.cpp
+++ b/Engine/Plugins/Animation/SkeletalMeshModelingTools/Source/SkeletalMeshModelingTools/Private/SkeletalMeshModelingToolsCommands.cpp
@@ -13,6 +13,8 @@ void FSkeletalMeshModelingToolsCommands::RegisterCommands()
 	UI_COMMAND(ToggleEditingToolsMode, "Enable Editing Tools", "Toggles editing tools on or off.", EUserInterfaceActionType::ToggleButton, FInputChord());
 	
 	UI_COMMAND(NewMorphTarget, "New Morph Target", "Add new Morph Target.", EUserInterfaceActionType::Button, FInputChord());
+	
+	UI_COMMAND(MirrorMorphTarget, "Mirror Selected Morph Target", "Mirror(Symmetrize) Selected Morph Target", EUserInterfaceActionType::Button, FInputChord());
 }
 
 const FSkeletalMeshModelingToolsCommands& FSkeletalMeshModelingToolsCommands::Get()
diff --git a/Engine/Plugins/Animation/SkeletalMeshModelingTools/Source/SkeletalMeshModelingTools/Private/SkeletalMeshModelingToolsCommands.h b/Engine/Plugins/Animation/SkeletalMeshModelingTools/Source/SkeletalMeshModelingTools/Private/SkeletalMeshModelingToolsCommands.h
--- a/Engine/Plugins/Animation/SkeletalMeshModelingTools/Source/SkeletalMeshModelingTools/Private/SkeletalMeshModelingToolsCommands.h
+++ b/Engine/Plugins/Animation/SkeletalMeshModelingTools/Source/SkeletalMeshModelingTools/Private/SkeletalMeshModelingToolsCommands.h
@@ -28,6 +28,8 @@ public:
 
 	// Morph Target Editing commands
 	TSharedPtr<FUICommandInfo> NewMorphTarget;
+
+	TSharedPtr<FUICommandInfo> MirrorMorphTarget;
 };
 
 class FSkeletalMeshModelingToolsActionCommands : public TInteractiveToolCommands<FSkeletalMeshModelingToolsActionCommands>
diff --git a/Engine/Plugins/Animation/SkeletalMeshModelingTools/Source/SkeletalMeshModelingTools/Private/SkeletalMeshModelingToolsEditorMode.cpp b/Engine/Plugins/Animation/SkeletalMeshModelingTools/Source/SkeletalMeshModelingTools/Private/SkeletalMeshModelingToolsEditorMode.cpp
--- a/Engine/Plugins/Animation/SkeletalMeshModelingTools/Source/SkeletalMeshModelingTools/Private/SkeletalMeshModelingToolsEditorMode.cpp
+++ b/Engine/Plugins/Animation/SkeletalMeshModelingTools/Source/SkeletalMeshModelingTools/Private/SkeletalMeshModelingToolsEditorMode.cpp
@@ -621,8 +621,8 @@ void USkeletalMeshModelingToolsEditorMode::HandleSetEditingMorphTarget(FName InM
 	FName OldEditingMorphTarget = EditingMorphTarget;
 	EditingMorphTarget = InMorphTarget;
 
-	GetCurrentEditingCache()->OverrideMorphTargetWeight(EditingMorphTarget, 1.0f);
 	GetCurrentEditingCache()->ClearMorphTargetOverride(OldEditingMorphTarget);
+	GetCurrentEditingCache()->OverrideMorphTargetWeight(EditingMorphTarget, 1.0f);
 }
 
 
@@ -670,9 +670,13 @@ TArray<FName> USkeletalMeshModelingToolsEditorMode::HandleDuplicateMorphTargets(
 	return GetCurrentEditingCache()->DuplicateMorphTargets(Names);
 }
 
+void USkeletalMeshModelingToolsEditorMode::HandleMirrorSelectedMorphTarget(FName Name)
+{
+	GetCurrentEditingCache()->MirrorMorphTargets({Name});
+}
+
 void USkeletalMeshModelingToolsEditorMode::HandleSetMorphTargetWeight(FName MorphTarget, float Weight)
 {
-	
 	return GetCurrentEditingCache()->HandleSetMorphTargetWeight(MorphTarget, Weight);
 }
 
diff --git a/Engine/Plugins/Animation/SkeletalMeshModelingTools/Source/SkeletalMeshModelingTools/Private/SkeletalMeshModelingToolsEditorMode.h b/Engine/Plugins/Animation/SkeletalMeshModelingTools/Source/SkeletalMeshModelingTools/Private/SkeletalMeshModelingToolsEditorMode.h
--- a/Engine/Plugins/Animation/SkeletalMeshModelingTools/Source/SkeletalMeshModelingTools/Private/SkeletalMeshModelingToolsEditorMode.h
+++ b/Engine/Plugins/Animation/SkeletalMeshModelingTools/Source/SkeletalMeshModelingTools/Private/SkeletalMeshModelingToolsEditorMode.h
@@ -168,6 +168,7 @@ public:
 	FName HandleRenameMorphTarget(FName OldName, FName NameName);
 	void HandleRemoveMorphTargets(const TArray<FName>& Names);
 	TArray<FName> HandleDuplicateMorphTargets(const TArray<FName>& Names);
+	void HandleMirrorSelectedMorphTarget(FName Name);	
 	
 	void HandleSetMorphTargetWeight(FName MorphTarget, float Weight);
 	bool GetMorphTargetAutoFill(FName MorphTarget);
diff --git a/Engine/Plugins/Animation/SkeletalMeshModelingTools/Source/SkeletalMeshModelingTools/Private/SkeletalMeshModelingToolsEditorModeToolkit.cpp b/Engine/Plugins/Animation/SkeletalMeshModelingTools/Source/SkeletalMeshModelingTools/Private/SkeletalMeshModelingToolsEditorModeToolkit.cpp
--- a/Engine/Plugins/Animation/SkeletalMeshModelingTools/Source/SkeletalMeshModelingTools/Private/SkeletalMeshModelingToolsEditorModeToolkit.cpp
+++ b/Engine/Plugins/Animation/SkeletalMeshModelingTools/Source/SkeletalMeshModelingTools/Private/SkeletalMeshModelingToolsEditorModeToolkit.cpp
@@ -607,6 +607,7 @@ void FSkeletalMeshModelingToolsEditorModeToolkit::OnEditorModeInitialized()
 	Delegates.OnRenameMorphTarget.BindUObject(EditorMode.Get(), &USkeletalMeshModelingToolsEditorMode::HandleRenameMorphTarget);
 	Delegates.OnRemoveMorphTargets.BindUObject(EditorMode.Get(), &USkeletalMeshModelingToolsEditorMode::HandleRemoveMorphTargets);
 	Delegates.OnDuplicateMorphTargets.BindUObject(EditorMode.Get(), &USkeletalMeshModelingToolsEditorMode::HandleDuplicateMorphTargets);
+	Delegates.OnMirrorSelectedMorphTarget.BindUObject(EditorMode.Get(), &USkeletalMeshModelingToolsEditorMode::HandleMirrorSelectedMorphTarget);
 	
 	SAssignNew(MorphTargetManager, SMorphTargetManager)
 			.Delegates(Delegates);
diff --git a/Engine/Plugins/Animation/SkeletalMeshModelingTools/Source/SkeletalMeshModelingTools/Private/ToolTargets/SKMBackedDynaMeshComponentToolTarget.cpp b/Engine/Plugins/Animation/SkeletalMeshModelingTools/Source/SkeletalMeshModelingTools/Private/ToolTargets/SKMBackedDynaMeshComponentToolTarget.cpp
--- a/Engine/Plugins/Animation/SkeletalMeshModelingTools/Source/SkeletalMeshModelingTools/Private/ToolTargets/SKMBackedDynaMeshComponentToolTarget.cpp
+++ b/Engine/Plugins/Animation/SkeletalMeshModelingTools/Source/SkeletalMeshModelingTools/Private/ToolTargets/SKMBackedDynaMeshComponentToolTarget.cpp
@@ -46,6 +46,9 @@ void USkeletalMeshBackedDynamicMeshComponentToolTarget::CommitDynamicMeshChange(
 	}
 	
 	Super::CommitDynamicMeshChange(MoveTemp(Change), ChangeMessage);
+	// we don't yet have a way of knowing if the tool changes the base geometry or not (either topo or just vert positions)
+	// so we have to assume that it is dirty such that symmetry tracker is always up to date.
+	GetComponent()->MarkBaseGeometryDirty();
 	GetComponent()->MarkDirty();
 
 	if (bCommitToSkeletalMeshOnToolCommit)
diff --git a/Engine/Plugins/Animation/SkeletalMeshModelingTools/Source/SkeletalMeshModelingTools/SkeletalMeshModelingTools.Build.cs b/Engine/Plugins/Animation/SkeletalMeshModelingTools/Source/SkeletalMeshModelingTools/SkeletalMeshModelingTools.Build.cs
--- a/Engine/Plugins/Animation/SkeletalMeshModelingTools/Source/SkeletalMeshModelingTools/SkeletalMeshModelingTools.Build.cs
+++ b/Engine/Plugins/Animation/SkeletalMeshModelingTools/Source/SkeletalMeshModelingTools/SkeletalMeshModelingTools.Build.cs
@@ -45,7 +45,8 @@ public class SkeletalMeshModelingTools : ModuleRules
 				"InterchangeEngine",
 				"GeometryFramework",
 				"EditorWidgets",
-				"DeveloperSettings"
+				"DeveloperSettings",
+				"DynamicMesh"
 			}
 		);
 
