[SKMModelingTool] Added hotkey "M" for mirror currently editing morph target

#Synced-CL 50896465

diff --git a/Engine/Plugins/Animation/SkeletalMeshModelingTools/Source/SkeletalMeshModelingTools/Private/SMorphTargetManager.cpp b/Engine/Plugins/Animation/SkeletalMeshModelingTools/Source/SkeletalMeshModelingTools/Private/SMorphTargetManager.cpp
--- a/Engine/Plugins/Animation/SkeletalMeshModelingTools/Source/SkeletalMeshModelingTools/Private/SMorphTargetManager.cpp
+++ b/Engine/Plugins/Animation/SkeletalMeshModelingTools/Source/SkeletalMeshModelingTools/Private/SMorphTargetManager.cpp
@@ -355,7 +355,7 @@ void SMorphTargetManager::BindCommands()
 		FCanExecuteAction::CreateSP(this, &SMorphTargetManager::CanDuplicate));
 	
 	CommandList->MapAction(
-		Commands.MirrorMorphTarget,
+		Commands.MirrorSelectedMorphTarget,
 		FExecuteAction::CreateSP(this, &SMorphTargetManager::MirrorSelectedMorphTarget),
 		FCanExecuteAction::CreateSP(this, &SMorphTargetManager::CanMirror));
 }
@@ -473,7 +473,7 @@ TSharedPtr<SWidget> SMorphTargetManager::OnGetContextMenuContent()
 	MenuBuilder.AddMenuEntry(FGenericCommands::Get().Rename);
 	MenuBuilder.AddMenuEntry(FGenericCommands::Get().Delete);
 	MenuBuilder.AddMenuEntry(FGenericCommands::Get().Duplicate);
-	MenuBuilder.AddMenuEntry(FSkeletalMeshModelingToolsCommands::Get().MirrorMorphTarget);
+	MenuBuilder.AddMenuEntry(FSkeletalMeshModelingToolsCommands::Get().MirrorSelectedMorphTarget);
 	
 	MenuBuilder.EndSection();
 	return MenuBuilder.MakeWidget();
diff --git a/Engine/Plugins/Animation/SkeletalMeshModelingTools/Source/SkeletalMeshModelingTools/Private/SkeletalMeshModelingToolsCommands.cpp b/Engine/Plugins/Animation/SkeletalMeshModelingTools/Source/SkeletalMeshModelingTools/Private/SkeletalMeshModelingToolsCommands.cpp
--- a/Engine/Plugins/Animation/SkeletalMeshModelingTools/Source/SkeletalMeshModelingTools/Private/SkeletalMeshModelingToolsCommands.cpp
+++ b/Engine/Plugins/Animation/SkeletalMeshModelingTools/Source/SkeletalMeshModelingTools/Private/SkeletalMeshModelingToolsCommands.cpp
@@ -14,7 +14,9 @@ void FSkeletalMeshModelingToolsCommands::RegisterCommands()
 	
 	UI_COMMAND(NewMorphTarget, "New Morph Target", "Add new Morph Target.", EUserInterfaceActionType::Button, FInputChord());
 	
-	UI_COMMAND(MirrorMorphTarget, "Mirror Selected Morph Target", "Mirror(Symmetrize) Selected Morph Target", EUserInterfaceActionType::Button, FInputChord());
+	UI_COMMAND(MirrorSelectedMorphTarget, "Mirror Selected Morph Target", "Mirror(Symmetrize) Selected Morph Target", EUserInterfaceActionType::Button, FInputChord());
+	
+	UI_COMMAND(MirrorEditingMorphTarget, "Mirror Editing Morph Target", "Mirror(Symmetrize) Currently editing Morph Target", EUserInterfaceActionType::Button, FInputChord(EKeys::M));
 }
 
 const FSkeletalMeshModelingToolsCommands& FSkeletalMeshModelingToolsCommands::Get()
diff --git a/Engine/Plugins/Animation/SkeletalMeshModelingTools/Source/SkeletalMeshModelingTools/Private/SkeletalMeshModelingToolsCommands.h b/Engine/Plugins/Animation/SkeletalMeshModelingTools/Source/SkeletalMeshModelingTools/Private/SkeletalMeshModelingToolsCommands.h
--- a/Engine/Plugins/Animation/SkeletalMeshModelingTools/Source/SkeletalMeshModelingTools/Private/SkeletalMeshModelingToolsCommands.h
+++ b/Engine/Plugins/Animation/SkeletalMeshModelingTools/Source/SkeletalMeshModelingTools/Private/SkeletalMeshModelingToolsCommands.h
@@ -29,7 +29,9 @@ public:
 	// Morph Target Editing commands
 	TSharedPtr<FUICommandInfo> NewMorphTarget;
 
-	TSharedPtr<FUICommandInfo> MirrorMorphTarget;
+	TSharedPtr<FUICommandInfo> MirrorSelectedMorphTarget;
+	
+	TSharedPtr<FUICommandInfo> MirrorEditingMorphTarget;
 };
 
 class FSkeletalMeshModelingToolsActionCommands : public TInteractiveToolCommands<FSkeletalMeshModelingToolsActionCommands>
diff --git a/Engine/Plugins/Animation/SkeletalMeshModelingTools/Source/SkeletalMeshModelingTools/Private/SkeletalMeshModelingToolsEditorMode.cpp b/Engine/Plugins/Animation/SkeletalMeshModelingTools/Source/SkeletalMeshModelingTools/Private/SkeletalMeshModelingToolsEditorMode.cpp
--- a/Engine/Plugins/Animation/SkeletalMeshModelingTools/Source/SkeletalMeshModelingTools/Private/SkeletalMeshModelingToolsEditorMode.cpp
+++ b/Engine/Plugins/Animation/SkeletalMeshModelingTools/Source/SkeletalMeshModelingTools/Private/SkeletalMeshModelingToolsEditorMode.cpp
@@ -670,7 +670,7 @@ TArray<FName> USkeletalMeshModelingToolsEditorMode::HandleDuplicateMorphTargets(
 	return GetCurrentEditingCache()->DuplicateMorphTargets(Names);
 }
 
-void USkeletalMeshModelingToolsEditorMode::HandleMirrorSelectedMorphTarget(FName Name)
+void USkeletalMeshModelingToolsEditorMode::HandleMirrorMorphTarget(FName Name)
 {
 	GetCurrentEditingCache()->MirrorMorphTargets({Name});
 }
@@ -968,6 +968,26 @@ void USkeletalMeshModelingToolsEditorMode::InitializeGeometrySelectionSystems()
 			return WeakThis->SelectionManager->IsSoftSelectionEnabled();
 		}));
 	CommandList->MapAction(ToolManagerCommands.SelectionEnableSoftSelection, ToggleSoftSelectionAction);
+
+	// Register Morph Target operations
+	FUIAction MirrorEditingMorphTargetAction = FUIAction(
+		FExecuteAction::CreateLambda([WeakThis]
+		{
+			if (!WeakThis.IsValid())
+			{
+				return;
+			}
+			WeakThis->HandleMirrorMorphTarget(WeakThis->GetEditingMorphTarget());
+		}),
+		FCanExecuteAction::CreateLambda([WeakThis]()
+		{
+			if (!WeakThis.IsValid())
+			{
+				return false;
+			}
+			return WeakThis->GetEditingMorphTarget() != NAME_None;
+		}));
+	CommandList->MapAction(FSkeletalMeshModelingToolsCommands::Get().MirrorEditingMorphTarget, MirrorEditingMorphTargetAction);
 	
 	// set up selection type toggles
 	auto RegisterSelectionMode = [WeakThis](UGeometrySelectionManager::EMeshTopologyMode TopoMode, UE::Geometry::EGeometryElementType ElementMode, TSharedPtr<FUICommandInfo> UICommand)
@@ -1147,6 +1167,15 @@ void USkeletalMeshModelingToolsEditorMode::PopulateViewportToolbar()
 						SoftSelectionSection.AddEntry(FToolMenuEntry::InitWidget(
 							"SoftSelectionRadius", SoftSelectionRadiusWidget, LOCTEXT("SoftSelectionRadiusLabel", "Radius")));
 				}
+
+				// Section for Morph Target
+				{
+					FToolMenuSection& MorphTargetSection =
+						Submenu->FindOrAddSection("Morph Target", LOCTEXT("MorphTargetLabel", "Morph Target"));
+					MorphTargetSection.AddEntry(
+							FToolMenuEntry::InitMenuEntryWithCommandList(
+								FSkeletalMeshModelingToolsCommands::Get().MirrorEditingMorphTarget, CommandList));
+				}
 			}
 		})
 	);
diff --git a/Engine/Plugins/Animation/SkeletalMeshModelingTools/Source/SkeletalMeshModelingTools/Private/SkeletalMeshModelingToolsEditorMode.h b/Engine/Plugins/Animation/SkeletalMeshModelingTools/Source/SkeletalMeshModelingTools/Private/SkeletalMeshModelingToolsEditorMode.h
--- a/Engine/Plugins/Animation/SkeletalMeshModelingTools/Source/SkeletalMeshModelingTools/Private/SkeletalMeshModelingToolsEditorMode.h
+++ b/Engine/Plugins/Animation/SkeletalMeshModelingTools/Source/SkeletalMeshModelingTools/Private/SkeletalMeshModelingToolsEditorMode.h
@@ -168,7 +168,7 @@ public:
 	FName HandleRenameMorphTarget(FName OldName, FName NameName);
 	void HandleRemoveMorphTargets(const TArray<FName>& Names);
 	TArray<FName> HandleDuplicateMorphTargets(const TArray<FName>& Names);
-	void HandleMirrorSelectedMorphTarget(FName Name);	
+	void HandleMirrorMorphTarget(FName Name);	
 	
 	void HandleSetMorphTargetWeight(FName MorphTarget, float Weight);
 	bool GetMorphTargetAutoFill(FName MorphTarget);
diff --git a/Engine/Plugins/Animation/SkeletalMeshModelingTools/Source/SkeletalMeshModelingTools/Private/SkeletalMeshModelingToolsEditorModeToolkit.cpp b/Engine/Plugins/Animation/SkeletalMeshModelingTools/Source/SkeletalMeshModelingTools/Private/SkeletalMeshModelingToolsEditorModeToolkit.cpp
--- a/Engine/Plugins/Animation/SkeletalMeshModelingTools/Source/SkeletalMeshModelingTools/Private/SkeletalMeshModelingToolsEditorModeToolkit.cpp
+++ b/Engine/Plugins/Animation/SkeletalMeshModelingTools/Source/SkeletalMeshModelingTools/Private/SkeletalMeshModelingToolsEditorModeToolkit.cpp
@@ -607,7 +607,7 @@ void FSkeletalMeshModelingToolsEditorModeToolkit::OnEditorModeInitialized()
 	Delegates.OnRenameMorphTarget.BindUObject(EditorMode.Get(), &USkeletalMeshModelingToolsEditorMode::HandleRenameMorphTarget);
 	Delegates.OnRemoveMorphTargets.BindUObject(EditorMode.Get(), &USkeletalMeshModelingToolsEditorMode::HandleRemoveMorphTargets);
 	Delegates.OnDuplicateMorphTargets.BindUObject(EditorMode.Get(), &USkeletalMeshModelingToolsEditorMode::HandleDuplicateMorphTargets);
-	Delegates.OnMirrorSelectedMorphTarget.BindUObject(EditorMode.Get(), &USkeletalMeshModelingToolsEditorMode::HandleMirrorSelectedMorphTarget);
+	Delegates.OnMirrorSelectedMorphTarget.BindUObject(EditorMode.Get(), &USkeletalMeshModelingToolsEditorMode::HandleMirrorMorphTarget);
 	
 	SAssignNew(MorphTargetManager, SMorphTargetManager)
 			.Delegates(Delegates);
