[SKMModelingTools] fixed a race condition in mirror morph target operation

#Synced-CL 50896465

diff --git a/Engine/Plugins/Animation/SkeletalMeshModelingTools/Source/SkeletalMeshModelingTools/Private/Components/SKMBackedDynaMeshComponent.cpp b/Engine/Plugins/Animation/SkeletalMeshModelingTools/Source/SkeletalMeshModelingTools/Private/Components/SKMBackedDynaMeshComponent.cpp
--- a/Engine/Plugins/Animation/SkeletalMeshModelingTools/Source/SkeletalMeshModelingTools/Private/Components/SKMBackedDynaMeshComponent.cpp
+++ b/Engine/Plugins/Animation/SkeletalMeshModelingTools/Source/SkeletalMeshModelingTools/Private/Components/SKMBackedDynaMeshComponent.cpp
@@ -627,28 +627,36 @@ void USkeletalMeshBackedDynamicMeshComponent::MirrorMorphTargets(const TArray<FN
 				FVector Delta;
 				SourceMorphTargetAttribute->GetValue(VertID, Delta);
 
-				// Making sure on-plane verts are dealt with first
-				TArray<FVector> ConstrainedPosition;
-				ConstrainedPosition.Add(Position+Delta);
-				ReadOnlySymmetry->ApplySymmetryPlaneConstraints({VertID}, ConstrainedPosition);
-
-				Delta = ConstrainedPosition[0] - Position;
-				SourceMorphTargetAttribute->SetValue(VertID, Delta);
+				constexpr bool bForceSameSizeWithGaps = true;
+				TArray<int32> MirroredVertID;
+				ReadOnlySymmetry->GetMirrorVertexROI({VertID}, MirroredVertID, bForceSameSizeWithGaps);
 
-				// Process positive side only	
-				if (Position.X < 0)
+				// Making sure on-plane verts are dealt with first 
+				// (it could also be that the vert is not on-plane and does not have a match, which will simply be ignored)
+				if (MirroredVertID[0] == INDEX_NONE)
 				{
-					return;
-				}
-
-				FVector MirroredDelta = ReadOnlySymmetry->GetMirroredAxis(Delta);
-				
-				TArray<int32> MirroredVertID;
-				ReadOnlySymmetry->GetMirrorVertexROI({VertID}, MirroredVertID, true);
+					
+					TArray<FVector> ConstrainedPosition;
+					ConstrainedPosition.Add(Position+Delta);
+					ReadOnlySymmetry->ApplySymmetryPlaneConstraints({VertID}, ConstrainedPosition);
 
-				if (MirroredVertID[0] != INDEX_NONE)
+					Delta = ConstrainedPosition[0] - Position;
+					SourceMorphTargetAttribute->SetValue(VertID, Delta);	
+				}
+				else
 				{
-					SourceMorphTargetAttribute->SetValue(MirroredVertID[0], MirroredDelta);
+					// Process positive side only	
+					if (Position.X < 0)
+					{
+						return;
+					}
+
+					FVector MirroredDelta = ReadOnlySymmetry->GetMirroredAxis(Delta);
+
+					if (MirroredVertID[0] != INDEX_NONE)
+					{
+						SourceMorphTargetAttribute->SetValue(MirroredVertID[0], MirroredDelta);
+					}	
 				}
 			});
 		}
