[SKMModelingTools] Fixed undo/redo not working for mirror editing morph target

#Synced-CL 50896465

#Synced-CL 50896465

diff --git a/Engine/Plugins/Animation/SkeletalMeshModelingTools/Source/SkeletalMeshModelingTools/Private/SkeletalMeshModelingToolsEditorMode.cpp b/Engine/Plugins/Animation/SkeletalMeshModelingTools/Source/SkeletalMeshModelingTools/Private/SkeletalMeshModelingToolsEditorMode.cpp
--- a/Engine/Plugins/Animation/SkeletalMeshModelingTools/Source/SkeletalMeshModelingTools/Private/SkeletalMeshModelingToolsEditorMode.cpp
+++ b/Engine/Plugins/Animation/SkeletalMeshModelingTools/Source/SkeletalMeshModelingTools/Private/SkeletalMeshModelingToolsEditorMode.cpp
@@ -973,11 +973,19 @@ void USkeletalMeshModelingToolsEditorMode::InitializeGeometrySelectionSystems()
 	FUIAction MirrorEditingMorphTargetAction = FUIAction(
 		FExecuteAction::CreateLambda([WeakThis]
 		{
-			if (!WeakThis.IsValid())
+			if (!WeakThis.IsValid() || !WeakThis->GetToolManager() || !WeakThis->GetToolManager()->GetContextTransactionsAPI())
 			{
 				return;
 			}
-			WeakThis->HandleMirrorMorphTarget(WeakThis->GetEditingMorphTarget());
+				
+			FName MorphTargetName = WeakThis->GetEditingMorphTarget();
+			FText Title = FText::Format( LOCTEXT("MirrorEditingMorphTransaction", "Mirror(Symmetrize) {0}"), FText::FromName(MorphTargetName));
+			IToolsContextTransactionsAPI* TransactionsAPI = WeakThis->GetToolManager()->GetContextTransactionsAPI();
+			TransactionsAPI->BeginUndoTransaction(Title);
+				
+			WeakThis->HandleMirrorMorphTarget(MorphTargetName);
+				
+			TransactionsAPI->EndUndoTransaction();
 		}),
 		FCanExecuteAction::CreateLambda([WeakThis]()
 		{
