[SKMModelingTool] adjust wording of menu options related to mirror morph target

#Synced-CL 50896465

diff --git a/Engine/Plugins/Animation/SkeletalMeshModelingTools/Source/SkeletalMeshModelingTools/Private/SMorphTargetManager.cpp b/Engine/Plugins/Animation/SkeletalMeshModelingTools/Source/SkeletalMeshModelingTools/Private/SMorphTargetManager.cpp
--- a/Engine/Plugins/Animation/SkeletalMeshModelingTools/Source/SkeletalMeshModelingTools/Private/SMorphTargetManager.cpp
+++ b/Engine/Plugins/Animation/SkeletalMeshModelingTools/Source/SkeletalMeshModelingTools/Private/SMorphTargetManager.cpp
@@ -591,7 +591,7 @@ void SMorphTargetManager::MirrorSelectedMorphTarget()
 	check(Items.Num() == 1);
 
 	FName MorphTargetName = Items[0]->Name;
-	FText Title = FText::Format( LOCTEXT("MirrorMorphTransaction", "Mirror(Symmetrize) {0}"), FText::FromName(MorphTargetName));
+	FText Title = FText::Format( LOCTEXT("MirrorMorphTransaction", "Mirror {0}"), FText::FromName(MorphTargetName));
 	FScopedTransaction Transaction(Title);
 
 	SetEditingMorphTarget(MorphTargetName);
diff --git a/Engine/Plugins/Animation/SkeletalMeshModelingTools/Source/SkeletalMeshModelingTools/Private/SkeletalMeshModelingToolsCommands.cpp b/Engine/Plugins/Animation/SkeletalMeshModelingTools/Source/SkeletalMeshModelingTools/Private/SkeletalMeshModelingToolsCommands.cpp
--- a/Engine/Plugins/Animation/SkeletalMeshModelingTools/Source/SkeletalMeshModelingTools/Private/SkeletalMeshModelingToolsCommands.cpp
+++ b/Engine/Plugins/Animation/SkeletalMeshModelingTools/Source/SkeletalMeshModelingTools/Private/SkeletalMeshModelingToolsCommands.cpp
@@ -12,11 +12,11 @@ void FSkeletalMeshModelingToolsCommands::RegisterCommands()
 {
 	UI_COMMAND(ToggleEditingToolsMode, "Enable Editing Tools", "Toggles editing tools on or off.", EUserInterfaceActionType::ToggleButton, FInputChord());
 	
-	UI_COMMAND(NewMorphTarget, "New Morph Target", "Add new Morph Target.", EUserInterfaceActionType::Button, FInputChord());
+	UI_COMMAND(NewMorphTarget, "New Morph Target", "Add a new Morph Target.", EUserInterfaceActionType::Button, FInputChord());
 	
-	UI_COMMAND(MirrorSelectedMorphTarget, "Mirror Selected Morph Target", "Mirror(Symmetrize) Selected Morph Target", EUserInterfaceActionType::Button, FInputChord());
+	UI_COMMAND(MirrorSelectedMorphTarget, "Mirror Selected Morph Target", "Mirror the selected Morph Target", EUserInterfaceActionType::Button, FInputChord());
 	
-	UI_COMMAND(MirrorEditingMorphTarget, "Mirror Editing Morph Target", "Mirror(Symmetrize) Currently editing Morph Target", EUserInterfaceActionType::Button, FInputChord(EKeys::M));
+	UI_COMMAND(MirrorEditingMorphTarget, "Mirror Editing Morph Target", "Mirror the Morph Target currently selected for editing, if geometry selection is active", EUserInterfaceActionType::Button, FInputChord(EKeys::M));
 }
 
 const FSkeletalMeshModelingToolsCommands& FSkeletalMeshModelingToolsCommands::Get()
diff --git a/Engine/Plugins/Animation/SkeletalMeshModelingTools/Source/SkeletalMeshModelingTools/Private/SkeletalMeshModelingToolsEditorMode.cpp b/Engine/Plugins/Animation/SkeletalMeshModelingTools/Source/SkeletalMeshModelingTools/Private/SkeletalMeshModelingToolsEditorMode.cpp
--- a/Engine/Plugins/Animation/SkeletalMeshModelingTools/Source/SkeletalMeshModelingTools/Private/SkeletalMeshModelingToolsEditorMode.cpp
+++ b/Engine/Plugins/Animation/SkeletalMeshModelingTools/Source/SkeletalMeshModelingTools/Private/SkeletalMeshModelingToolsEditorMode.cpp
@@ -979,7 +979,7 @@ void USkeletalMeshModelingToolsEditorMode::InitializeGeometrySelectionSystems()
 			}
 				
 			FName MorphTargetName = WeakThis->GetEditingMorphTarget();
-			FText Title = FText::Format( LOCTEXT("MirrorEditingMorphTransaction", "Mirror(Symmetrize) {0}"), FText::FromName(MorphTargetName));
+			FText Title = FText::Format( LOCTEXT("MirrorEditingMorphTransaction", "Mirror {0}"), FText::FromName(MorphTargetName));
 			IToolsContextTransactionsAPI* TransactionsAPI = WeakThis->GetToolManager()->GetContextTransactionsAPI();
 			TransactionsAPI->BeginUndoTransaction(Title);
 				
@@ -993,7 +993,7 @@ void USkeletalMeshModelingToolsEditorMode::InitializeGeometrySelectionSystems()
 			{
 				return false;
 			}
-			return WeakThis->GetEditingMorphTarget() != NAME_None;
+			return WeakThis->GetEditingMorphTarget() != NAME_None && WeakThis->IsGeometrySelectionActive();
 		}));
 	CommandList->MapAction(FSkeletalMeshModelingToolsCommands::Get().MirrorEditingMorphTarget, MirrorEditingMorphTargetAction);
 	
