[SKMModelingTools] Lazy construction of mesh symmetry for morph target mirroring

#Synced-CL 50896465

diff --git a/Engine/Plugins/Animation/SkeletalMeshModelingTools/Source/SkeletalMeshModelingTools/Private/Components/SKMBackedDynaMeshComponent.cpp b/Engine/Plugins/Animation/SkeletalMeshModelingTools/Source/SkeletalMeshModelingTools/Private/Components/SKMBackedDynaMeshComponent.cpp
--- a/Engine/Plugins/Animation/SkeletalMeshModelingTools/Source/SkeletalMeshModelingTools/Private/Components/SKMBackedDynaMeshComponent.cpp
+++ b/Engine/Plugins/Animation/SkeletalMeshModelingTools/Source/SkeletalMeshModelingTools/Private/Components/SKMBackedDynaMeshComponent.cpp
@@ -204,7 +204,6 @@ EMeshLODIdentifier USkeletalMeshBackedDynamicMeshComponent::Init(USkeletalMesh*
 
 	SetMesh(MoveTemp(LODMesh));
 
-	UpdateSymmetry();
 	// Additional Trackers
 	ResetTrackersDirect();
 
@@ -586,10 +585,8 @@ TArray<FName> USkeletalMeshBackedDynamicMeshComponent::DuplicateMorphTargets(con
 
 void USkeletalMeshBackedDynamicMeshComponent::MirrorMorphTargets(const TArray<FName>& InNames)
 {
-	if (!ReadOnlySymmetry)
-	{
-		return;
-	}
+	UpdateSymmetryIfNeeded();
+	check(ReadOnlySymmetry);
 	
 	TArray<FName> MorphsToEdit;
 	for (const FName& Name : InNames)
@@ -759,7 +756,6 @@ void USkeletalMeshBackedDynamicMeshComponent::MarkDirtyInternal()
 void USkeletalMeshBackedDynamicMeshComponent::MarkBaseGeometryDirtyInternal()
 {
 	Trackers.BaseGeometryChangeCount++;
-	UpdateSymmetry();
 }
 
 void USkeletalMeshBackedDynamicMeshComponent::MarkDirty()
@@ -792,10 +788,6 @@ void USkeletalMeshBackedDynamicMeshComponent::FTrackerChange::Apply(UObject* Obj
 {
 	if (USkeletalMeshBackedDynamicMeshComponent* Component = Cast<USkeletalMeshBackedDynamicMeshComponent>(Object))
 	{
-		if (Component->Trackers.BaseGeometryChangeCount != NewTrackers.BaseGeometryChangeCount)
-		{
-			Component->UpdateSymmetry();
-		}
 		Component->Trackers = NewTrackers;
 	}	
 }
@@ -804,10 +796,6 @@ void USkeletalMeshBackedDynamicMeshComponent::FTrackerChange::Revert(UObject* Ob
 {
 	if (USkeletalMeshBackedDynamicMeshComponent* Component = Cast<USkeletalMeshBackedDynamicMeshComponent>(Object))
 	{
-		if (Component->Trackers.BaseGeometryChangeCount != OldTrackers.BaseGeometryChangeCount)
-		{
-			Component->UpdateSymmetry();
-		}
 		Component->Trackers = OldTrackers;
 	}	
 }
@@ -855,11 +843,17 @@ USkeletalMeshBackedDynamicMeshComponent::FChangeScope::~FChangeScope()
 	}
 }
 
-void USkeletalMeshBackedDynamicMeshComponent::UpdateSymmetry()
+void USkeletalMeshBackedDynamicMeshComponent::UpdateSymmetryIfNeeded()
 {
 	using namespace UE::Geometry;
 
+	if (BaseGeometryVersionForCurrentSymmetry == Trackers.BaseGeometryChangeCount)
+	{
+		return;
+	}
 
+	BaseGeometryVersionForCurrentSymmetry = Trackers.BaseGeometryChangeCount;
+	
 	GetDynamicMesh()->ProcessMesh([this](const FDynamicMesh3& Mesh)
 	{
 		FAxisAlignedBox3d Bounds = Mesh.GetBounds(true);
diff --git a/Engine/Plugins/Animation/SkeletalMeshModelingTools/Source/SkeletalMeshModelingTools/Private/Components/SKMBackedDynaMeshComponent.h b/Engine/Plugins/Animation/SkeletalMeshModelingTools/Source/SkeletalMeshModelingTools/Private/Components/SKMBackedDynaMeshComponent.h
--- a/Engine/Plugins/Animation/SkeletalMeshModelingTools/Source/SkeletalMeshModelingTools/Private/Components/SKMBackedDynaMeshComponent.h
+++ b/Engine/Plugins/Animation/SkeletalMeshModelingTools/Source/SkeletalMeshModelingTools/Private/Components/SKMBackedDynaMeshComponent.h
@@ -187,7 +187,8 @@ protected:
 
 	// Query only symmetry, not allow to modify mesh
 	TPimplPtr<UE::Geometry::FMeshPlanarSymmetry> ReadOnlySymmetry;
-	void UpdateSymmetry();
+	int32 BaseGeometryVersionForCurrentSymmetry = -1;
+	void UpdateSymmetryIfNeeded();
 };
 
 #undef UE_API
