[SKMModelingTools] Fixed duplicate morph target not showing in the viewport immediately.

#Synced-CL 50896465

diff --git a/Engine/Plugins/Animation/SkeletalMeshModelingTools/Source/SkeletalMeshModelingTools/Private/Components/SKMBackedDynaMeshComponent.cpp b/Engine/Plugins/Animation/SkeletalMeshModelingTools/Source/SkeletalMeshModelingTools/Private/Components/SKMBackedDynaMeshComponent.cpp
--- a/Engine/Plugins/Animation/SkeletalMeshModelingTools/Source/SkeletalMeshModelingTools/Private/Components/SKMBackedDynaMeshComponent.cpp
+++ b/Engine/Plugins/Animation/SkeletalMeshModelingTools/Source/SkeletalMeshModelingTools/Private/Components/SKMBackedDynaMeshComponent.cpp
@@ -579,6 +579,7 @@ TArray<FName> USkeletalMeshBackedDynamicMeshComponent::DuplicateMorphTargets(con
 		Trackers.MorphTargetChangeTracker.HandleEditMorphTarget(Name);
 	}
 
+	MarkDirtyInternal();
 	
 	return NewMorphs;
 }
