[SkeletalMeshEditingTools] The Skeletal Mesh Editor now displays a menu in the viewport toolbar with options to set the local rotation axis length, thickness and to show all local rotation axes

#Synced-CL 51094013

diff --git a/Engine/Plugins/Animation/SkeletalMeshModelingTools/Source/SkeletalMeshModelingTools/Private/DetailCustomization/SkeletonEditingToolPropertyCustomizations.cpp b/Engine/Plugins/Animation/SkeletalMeshModelingTools/Source/SkeletalMeshModelingTools/Private/DetailCustomization/SkeletonEditingToolPropertyCustomizations.cpp
--- a/Engine/Plugins/Animation/SkeletalMeshModelingTools/Source/SkeletalMeshModelingTools/Private/DetailCustomization/SkeletonEditingToolPropertyCustomizations.cpp
+++ b/Engine/Plugins/Animation/SkeletalMeshModelingTools/Source/SkeletalMeshModelingTools/Private/DetailCustomization/SkeletonEditingToolPropertyCustomizations.cpp
@@ -187,9 +187,7 @@ const TArray<FName>& FSkeletonEditingPropertiesDetailCustomization::GetPropertie
 	static const TArray< FName > Properties(
 	{
 		GET_MEMBER_NAME_CHECKED(USkeletonEditingProperties, Name),
-		GET_MEMBER_NAME_CHECKED(USkeletonEditingProperties, bUpdateChildren),
-		GET_MEMBER_NAME_CHECKED(USkeletonEditingProperties, AxisLength),
-		GET_MEMBER_NAME_CHECKED(USkeletonEditingProperties, AxisThickness),
+		GET_MEMBER_NAME_CHECKED(USkeletonEditingProperties, bUpdateChildren)
 	});
 	return Properties;
 }
diff --git a/Engine/Plugins/Animation/SkeletalMeshModelingTools/Source/SkeletalMeshModelingTools/Private/SkeletalMeshModelingToolsCommands.cpp b/Engine/Plugins/Animation/SkeletalMeshModelingTools/Source/SkeletalMeshModelingTools/Private/SkeletalMeshModelingToolsCommands.cpp
--- a/Engine/Plugins/Animation/SkeletalMeshModelingTools/Source/SkeletalMeshModelingTools/Private/SkeletalMeshModelingToolsCommands.cpp
+++ b/Engine/Plugins/Animation/SkeletalMeshModelingTools/Source/SkeletalMeshModelingTools/Private/SkeletalMeshModelingToolsCommands.cpp
@@ -17,6 +17,8 @@ void FSkeletalMeshModelingToolsCommands::RegisterCommands()
 	UI_COMMAND(MirrorSelectedMorphTarget, "Mirror Selected Morph Target", "Mirror the selected Morph Target", EUserInterfaceActionType::Button, FInputChord());
 	
 	UI_COMMAND(MirrorEditingMorphTarget, "Mirror Editing Morph Target", "Mirror the Morph Target currently selected for editing, if geometry selection is active", EUserInterfaceActionType::Button, FInputChord(EKeys::M));
+
+	UI_COMMAND(ToggleShowAllLocalRotationAxes, "Show All Local Rotation Axes", "When enabled shows all local rotation axes, otherwise for the selected bone only", EUserInterfaceActionType::ToggleButton, FInputChord());	
 }
 
 const FSkeletalMeshModelingToolsCommands& FSkeletalMeshModelingToolsCommands::Get()
diff --git a/Engine/Plugins/Animation/SkeletalMeshModelingTools/Source/SkeletalMeshModelingTools/Private/SkeletalMeshModelingToolsCommands.h b/Engine/Plugins/Animation/SkeletalMeshModelingTools/Source/SkeletalMeshModelingTools/Private/SkeletalMeshModelingToolsCommands.h
--- a/Engine/Plugins/Animation/SkeletalMeshModelingTools/Source/SkeletalMeshModelingTools/Private/SkeletalMeshModelingToolsCommands.h
+++ b/Engine/Plugins/Animation/SkeletalMeshModelingTools/Source/SkeletalMeshModelingTools/Private/SkeletalMeshModelingToolsCommands.h
@@ -32,6 +32,9 @@ public:
 	TSharedPtr<FUICommandInfo> MirrorSelectedMorphTarget;
 	
 	TSharedPtr<FUICommandInfo> MirrorEditingMorphTarget;
+
+    /** Command to toggle show all local rotation axes */
+	TSharedPtr<FUICommandInfo> ToggleShowAllLocalRotationAxes;
 };
 
 class FSkeletalMeshModelingToolsActionCommands : public TInteractiveToolCommands<FSkeletalMeshModelingToolsActionCommands>
diff --git a/Engine/Plugins/Animation/SkeletalMeshModelingTools/Source/SkeletalMeshModelingTools/Private/SkeletalMeshModelingToolsEditorMode.cpp b/Engine/Plugins/Animation/SkeletalMeshModelingTools/Source/SkeletalMeshModelingTools/Private/SkeletalMeshModelingToolsEditorMode.cpp
--- a/Engine/Plugins/Animation/SkeletalMeshModelingTools/Source/SkeletalMeshModelingTools/Private/SkeletalMeshModelingToolsEditorMode.cpp
+++ b/Engine/Plugins/Animation/SkeletalMeshModelingTools/Source/SkeletalMeshModelingTools/Private/SkeletalMeshModelingToolsEditorMode.cpp
@@ -85,12 +85,10 @@
 
 #include "ModelingToolsEditorModeSettings.h"
 #include "Selection/GeometrySelectionManager.h"
-#include "ModelingToolsEditorModeStyle.h"
-#include "ToolMenus.h"
-#include "SkeletalMeshToolMenuContext.h"
 #include "Commands/ModifyGeometrySelectionCommand.h"
 #include "Selection/SKMPreviewDynaMeshSelector.h"
 #include "Widgets/Input/SSpinBox.h"
+#include "SkeletalMeshModelingToolsViewportToolbarExtension.h"
 
 #define LOCTEXT_NAMESPACE "SkeletalMeshModelingToolsEditorMode"
 
@@ -206,12 +204,6 @@ static void ShowEditorMessage(ELogVerbosity::Type InMessageType, const FText& In
 	FSlateNotificationManager::Get().AddNotification(Notification)->SetCompletionState(State);
 }
 
-namespace SkeletalMeshModelingToolsEditorModeLocals
-{
-	const FName ViewportToolbarExtensionOwnerName = "SkeletalMeshModelingToolsEditorModeViewportToolbar";
-}
-
-
 // NOTE: This is a simple proxy at the moment. In the future we want to pull in more of the 
 // modeling tools as we add support in the skelmesh storage.
 
@@ -293,6 +285,8 @@ void USkeletalMeshModelingToolsEditorMode::Initialize()
 
 void USkeletalMeshModelingToolsEditorMode::Enter()
 {
+	using namespace UE::SkeletalMeshModelingTools;
+
 	UEdMode::Enter();
 
 	UEditorInteractiveToolsContext* EditorInteractiveToolsContext = GetInteractiveToolsContext(EToolsContextScope::Editor);
@@ -386,7 +380,14 @@ void USkeletalMeshModelingToolsEditorMode::Enter()
 	GetInteractiveToolsContext()->ToolManager->OnToolEndedWithStatus.AddUObject(this, &USkeletalMeshModelingToolsEditorMode::OnToolEndedWithStatus);
 
 	InitializeGeometrySelectionSystems();
-	PopulateViewportToolbar();
+	
+	ExtendSkeletalMeshEditorViewportToolbar(*this, Toolkit->GetToolkitCommands());
+
+	const TWeakPtr<FModeToolkit> WeakToolkit = GetToolkit();
+	if (WeakToolkit.IsValid())
+	{
+		ExtendSkeletalMeshEditorViewportToolbar(*this, WeakToolkit.Pin()->GetToolkitCommands());
+	}
 
 	if (UModeManagerInteractiveToolsContext* const ModeManagerInteractiveToolsContext = GetModeManager()->GetInteractiveToolsContext())
 	{
@@ -698,13 +699,15 @@ void USkeletalMeshModelingToolsEditorMode::HandleSetMorphTargetAutoFill(FName Mo
 
 void USkeletalMeshModelingToolsEditorMode::Exit()
 {
+	using namespace UE::SkeletalMeshModelingTools;
+
 	if (UModeManagerInteractiveToolsContext* const ModeManagerInteractiveToolsContext = GetModeManager()->GetInteractiveToolsContext())
 	{
 		ModeManagerInteractiveToolsContext->OnBuildViewportInteractions().RemoveAll(this);
 		ModeManagerInteractiveToolsContext->PostBuildViewportInteractions().RemoveAll(this);
 	}
 	
-	RemoveViewportToolbarExtensions();
+	RemoveSkeletalMeshEditorViewportToolbarExtensions();
 
 
 	SelectionTransformTweaker->Shutdown();
@@ -1048,172 +1051,10 @@ void USkeletalMeshModelingToolsEditorMode::InitializeGeometrySelectionSystems()
 	}
 }
 
-void USkeletalMeshModelingToolsEditorMode::PopulateViewportToolbar()
-{
-	using namespace SkeletalMeshModelingToolsEditorModeLocals;
-	FToolMenuOwnerScoped ScopeOwner(ViewportToolbarExtensionOwnerName);
-	UToolMenu* const Menu = UToolMenus::Get()->ExtendMenu("AnimationEditor.ViewportToolbar");
-
-	TWeakObjectPtr<USkeletalMeshModelingToolsEditorMode> WeakThis(this);
-	
-	FToolMenuSection& MeshElementSelectionMenu = Menu->FindOrAddSection("Left");
-	FToolMenuEntry Entry = FToolMenuEntry::InitSubMenu(
-		"Mesh Element Selection",
-		LOCTEXT("MeshElementSelectionSubmenuLabel", "Mesh Element Selection"),
-		LOCTEXT("MeshElementSelectionSubmenuTooltip", "Mesh Element Selection settings in the viewport"),
-		FNewToolMenuDelegate::CreateLambda([WeakThis, WeakCommandList = Toolkit->GetToolkitCommands().ToWeakPtr()](UToolMenu* Submenu) -> void
-		{
-			FToolMenuSection& MeshElementSelectionOptionsSection =
-			Submenu->FindOrAddSection("Element Selection", LOCTEXT("ElementSelectionLabel", "Element Selection"));
-			if (const TSharedPtr<const FUICommandList>& CommandList = WeakCommandList.Pin())
-			{
-				auto CreateToolMenuEntry = [&CommandList, &MeshElementSelectionOptionsSection](const TSharedPtr<FUICommandInfo>& Command)
-				{
-					FToolMenuEntry MESMode = FToolMenuEntry::InitMenuEntryWithCommandList( Command, CommandList );
-					MESMode.SetShowInToolbarTopLevel(true); // add button to toolbar
-					MeshElementSelectionOptionsSection.AddEntry(MESMode);
-				};
-
-				// add all Mesh Element Selection modes
-				CreateToolMenuEntry(FModelingToolsManagerCommands::Get().MeshSelectionModeAction_NoSelection);
-				MeshElementSelectionOptionsSection.AddSeparator(NAME_None);
-				CreateToolMenuEntry(FModelingToolsManagerCommands::Get().MeshSelectionModeAction_MeshVertices);
-				CreateToolMenuEntry(FModelingToolsManagerCommands::Get().MeshSelectionModeAction_MeshEdges);
-				CreateToolMenuEntry(FModelingToolsManagerCommands::Get().MeshSelectionModeAction_MeshTriangles);
-				MeshElementSelectionOptionsSection.AddSeparator(NAME_None);
-				CreateToolMenuEntry(FModelingToolsManagerCommands::Get().MeshSelectionModeAction_GroupCorners);
-				CreateToolMenuEntry(FModelingToolsManagerCommands::Get().MeshSelectionModeAction_GroupEdges);
-				CreateToolMenuEntry(FModelingToolsManagerCommands::Get().MeshSelectionModeAction_GroupFaces);
-
-				// Section for Selection Edits
-				{
-					FToolMenuSection& SelectionEditsSection =
-						Submenu->FindOrAddSection("SelectionEdits", LOCTEXT("SelectionEditsLabel", "Selection Edits"));
-
-					SelectionEditsSection.AddEntry(
-					FToolMenuEntry::InitMenuEntryWithCommandList(
-						FModelingToolsManagerCommands::Get().BeginSelectionAction_SelectAll, CommandList));
-					SelectionEditsSection.AddEntry(
-					FToolMenuEntry::InitMenuEntryWithCommandList(
-						FModelingToolsManagerCommands::Get().BeginSelectionAction_ExpandToConnected, CommandList));
-					SelectionEditsSection.AddEntry(
-					FToolMenuEntry::InitMenuEntryWithCommandList(
-						FModelingToolsManagerCommands::Get().BeginSelectionAction_Invert, CommandList));
-					SelectionEditsSection.AddEntry(
-					FToolMenuEntry::InitMenuEntryWithCommandList(
-						FModelingToolsManagerCommands::Get().BeginSelectionAction_InvertConnected, CommandList));
-					SelectionEditsSection.AddEntry(
-					FToolMenuEntry::InitMenuEntryWithCommandList(
-						FModelingToolsManagerCommands::Get().BeginSelectionAction_Expand, CommandList));
-					SelectionEditsSection.AddEntry(
-					FToolMenuEntry::InitMenuEntryWithCommandList(
-						FModelingToolsManagerCommands::Get().BeginSelectionAction_Contract, CommandList));
-				}
-
-				
-				// Section for Local Frame Mode setting
-				{
-					FToolMenuSection& LocalFrameModeSection =
-						Submenu->FindOrAddSection("Local Frame Mode", LOCTEXT("LocalFrameModeLabel", "Local Frame Mode"));
-					LocalFrameModeSection.AddEntry(
-						FToolMenuEntry::InitMenuEntryWithCommandList(
-							FModelingToolsManagerCommands::Get().SelectionLocalFrameMode_Geometry, CommandList));
-					LocalFrameModeSection.AddEntry(
-						FToolMenuEntry::InitMenuEntryWithCommandList(
-							FModelingToolsManagerCommands::Get().SelectionLocalFrameMode_Object, CommandList));
-				}
-
-				// Section for Selection Filters
-				{
-					FToolMenuSection& SelectionFilterSection =
-						Submenu->FindOrAddSection("Selection Filters", LOCTEXT("SelectionFiltersLabel", "Selection Filters"));
-					SelectionFilterSection.AddEntry(
-							FToolMenuEntry::InitMenuEntryWithCommandList(
-								FModelingToolsManagerCommands::Get().SelectionHitBackFaces, CommandList));
-				}
-
-				// Section for Selection Filters
-				{
-					FToolMenuSection& SoftSelectionSection =
-						Submenu->FindOrAddSection("Soft Selection", LOCTEXT("SoftSelectionLabel", "Soft Selection"));
-					SoftSelectionSection.AddEntry(
-							FToolMenuEntry::InitMenuEntryWithCommandList(
-								FModelingToolsManagerCommands::Get().SelectionEnableSoftSelection, CommandList));
-
-
-					
-					TSharedRef<SWidget> SoftSelectionRadiusWidget =
-						SNew(SHorizontalBox)
-						+ SHorizontalBox::Slot()
-						.FillWidth(0.9f)
-						[
-							SNew(SSpinBox<double>)
-								.MinValue(0.0f)
-								.Value_Lambda(
-									[WeakThis]() -> double 
-									{
-										if (!WeakThis.IsValid())
-										{
-											return 0;
-										}
-										return WeakThis->SelectionManager->GetSoftSelectionRadius();
-									}
-								)
-								.OnValueChanged_Lambda(
-									[WeakThis](double InValue)
-									{
-										if (WeakThis.IsValid())
-										{
-											WeakThis->SelectionManager->SetSoftSelectionRadius(InValue);
-										}
-									}
-								)
-						]
-						+ SHorizontalBox::Slot()
-						.FillWidth(0.1f);
-
-						SoftSelectionSection.AddEntry(FToolMenuEntry::InitWidget(
-							"SoftSelectionRadius", SoftSelectionRadiusWidget, LOCTEXT("SoftSelectionRadiusLabel", "Radius")));
-				}
-
-				// Section for Morph Target
-				{
-					FToolMenuSection& MorphTargetSection =
-						Submenu->FindOrAddSection("Morph Target", LOCTEXT("MorphTargetLabel", "Morph Target"));
-					MorphTargetSection.AddEntry(
-							FToolMenuEntry::InitMenuEntryWithCommandList(
-								FSkeletalMeshModelingToolsCommands::Get().MirrorEditingMorphTarget, CommandList));
-				}
-			}
-		})
-	);
-
-	Entry.ToolBarData.ResizeParams.ClippingPriority = 950;
-	Entry.ToolBarData.LabelOverride = FText();
-	Entry.InsertPosition = FToolMenuInsert("Transform", EToolMenuInsertType::After);
-	Entry.Icon = FSlateIcon(FModelingToolsEditorModeStyle::GetStyleSetName(), "ModelingToolsManagerCommands.MeshElementSelection");
-
-	MeshElementSelectionMenu.AddDynamicEntry(TEXT("Mesh Element Selection"),
-		FNewToolMenuSectionDelegate::CreateLambda([Entry](FToolMenuSection& InSection)
-		{
-			// Only extending the menu for skeletal mesh editor
-			const USkeletalMeshToolMenuContext* Context = InSection.FindContext<USkeletalMeshToolMenuContext>();
-			if (Context && Context->SkeletalMeshEditor.IsValid())
-			{
-				InSection.AddEntry(Entry);
-			}
-		}));
-}
-
-void USkeletalMeshModelingToolsEditorMode::RemoveViewportToolbarExtensions()
-{
-	UToolMenus::Get()->UnregisterOwnerByName(SkeletalMeshModelingToolsEditorModeLocals::ViewportToolbarExtensionOwnerName);
-}
-
 void USkeletalMeshModelingToolsEditorMode::OnBuildViewportInteractions(UViewportInteractionsBehaviorSource* Source)
 {
 	Source->AddInteraction<USkeletalMeshModelingToolsClickSelection>();
-	
+
 	USkeletalMeshModelingToolsGeometryFrustumSelection* GeometryFrustumSelection =
 		Source->AddInteraction<USkeletalMeshModelingToolsGeometryFrustumSelection>();
 	GeometryFrustumSelection->BindSelectionManager(SelectionManager);
diff --git a/Engine/Plugins/Animation/SkeletalMeshModelingTools/Source/SkeletalMeshModelingTools/Private/SkeletalMeshModelingToolsEditorMode.h b/Engine/Plugins/Animation/SkeletalMeshModelingTools/Source/SkeletalMeshModelingTools/Private/SkeletalMeshModelingToolsEditorMode.h
--- a/Engine/Plugins/Animation/SkeletalMeshModelingTools/Source/SkeletalMeshModelingTools/Private/SkeletalMeshModelingToolsEditorMode.h
+++ b/Engine/Plugins/Animation/SkeletalMeshModelingTools/Source/SkeletalMeshModelingTools/Private/SkeletalMeshModelingToolsEditorMode.h
@@ -283,8 +283,6 @@ private:
 	bool bAllowGeometrySelection = true;
 
 	void InitializeGeometrySelectionSystems();
-	void PopulateViewportToolbar();
-	void RemoveViewportToolbarExtensions();
 
 	void OnBuildViewportInteractions(UViewportInteractionsBehaviorSource* Source);
 	void PostBuildViewportInteractions(const UViewportInteractionsBehaviorSource* Source);
diff --git a/Engine/Plugins/Animation/SkeletalMeshModelingTools/Source/SkeletalMeshModelingTools/Private/SkeletalMeshModelingToolsEditorModeToolkit.cpp b/Engine/Plugins/Animation/SkeletalMeshModelingTools/Source/SkeletalMeshModelingTools/Private/SkeletalMeshModelingToolsEditorModeToolkit.cpp
--- a/Engine/Plugins/Animation/SkeletalMeshModelingTools/Source/SkeletalMeshModelingTools/Private/SkeletalMeshModelingToolsEditorModeToolkit.cpp
+++ b/Engine/Plugins/Animation/SkeletalMeshModelingTools/Source/SkeletalMeshModelingTools/Private/SkeletalMeshModelingToolsEditorModeToolkit.cpp
@@ -40,6 +40,7 @@
 #include "Widgets/Docking/SDockTab.h"
 #include "ScopedTransaction.h"
 #include "Selection/GeometrySelectionManager.h"
+#include "SkeletalMesh/SkeletalMeshModelingToolsEditorSettings.h"
 
 
 #define LOCTEXT_NAMESPACE "SkeletalMeshModelingToolsEditorModeToolkit"
@@ -66,6 +67,8 @@ void FSkeletalMeshModelingToolsEditorModeToolkit::Init(
 
 	FModeToolkit::Init(InToolkitHost, InOwningMode);
 
+	SetupCommands();
+
 	ModeWarningArea = SNew(STextBlock)
 		.AutoWrapText(true)
 		.Font(FCoreStyle::GetDefaultFontStyle("Bold", 9))
@@ -494,6 +497,29 @@ void FSkeletalMeshModelingToolsEditorModeToolkit::MakeToolAcceptCancelWidget()
 	];
 }
 
+void FSkeletalMeshModelingToolsEditorModeToolkit::SetupCommands()
+{
+	const FSkeletalMeshModelingToolsCommands& Commands = FSkeletalMeshModelingToolsCommands::Get();
+	
+	GetToolkitCommands()->MapAction(
+		Commands.ToggleShowAllLocalRotationAxes,
+		FUIAction(
+			FExecuteAction::CreateLambda([]()
+				{
+					USkeletalMeshModelingToolsEditorSettings* const Settings = GetMutableDefault<USkeletalMeshModelingToolsEditorSettings>();
+					Settings->SetShowAllLocalRotationAxes(!Settings->GetShowAllLocalRotationAxes());
+					Settings->SaveConfig();
+				}),
+			FCanExecuteAction(),
+			FIsActionChecked::CreateLambda([]()
+				{
+					const USkeletalMeshModelingToolsEditorSettings* const Settings = GetDefault<USkeletalMeshModelingToolsEditorSettings>();
+					return Settings->GetShowAllLocalRotationAxes();
+				})
+		)
+	);
+}
+
 void FSkeletalMeshModelingToolsEditorModeToolkit::PostNotification(const FText& InMessage)
 {
 	ClearNotification();
diff --git a/Engine/Plugins/Animation/SkeletalMeshModelingTools/Source/SkeletalMeshModelingTools/Private/SkeletalMeshModelingToolsEditorModeToolkit.h b/Engine/Plugins/Animation/SkeletalMeshModelingTools/Source/SkeletalMeshModelingTools/Private/SkeletalMeshModelingToolsEditorModeToolkit.h
--- a/Engine/Plugins/Animation/SkeletalMeshModelingTools/Source/SkeletalMeshModelingTools/Private/SkeletalMeshModelingToolsEditorModeToolkit.h
+++ b/Engine/Plugins/Animation/SkeletalMeshModelingTools/Source/SkeletalMeshModelingTools/Private/SkeletalMeshModelingToolsEditorModeToolkit.h
@@ -53,6 +53,8 @@ public:
 	
 	void RefreshMorphTargetManager();
 private:
+	void SetupCommands();
+
 	void PostNotification(const FText& Message);
 	void ClearNotification();
 
diff --git a/Engine/Plugins/Animation/SkeletalMeshModelingTools/Source/SkeletalMeshModelingTools/Private/SkeletalMeshModelingToolsViewportToolbarExtension.cpp b/Engine/Plugins/Animation/SkeletalMeshModelingTools/Source/SkeletalMeshModelingTools/Private/SkeletalMeshModelingToolsViewportToolbarExtension.cpp
new file mode 100644
--- /dev/null
+++ b/Engine/Plugins/Animation/SkeletalMeshModelingTools/Source/SkeletalMeshModelingTools/Private/SkeletalMeshModelingToolsViewportToolbarExtension.cpp
@@ -0,0 +1,334 @@
+// Copyright Epic Games, Inc. All Rights Reserved.
+
+#include "SkeletalMeshModelingToolsViewportToolbarExtension.h"
+
+#include "ISkeletalMeshEditor.h"
+#include "ModelingToolsEditorModeStyle.h"
+#include "ModelingToolsManagerActions.h"
+#include "Selection/GeometrySelectionManager.h"
+#include "SkeletalMesh/SkeletalMeshModelingToolsEditorSettings.h"
+#include "SkeletalMeshModelingToolsCommands.h"
+#include "SkeletalMeshModelingToolsEditorMode.h"
+#include "SkeletalMeshModelingToolsEditorModeToolkit.h"
+#include "SkeletalMeshToolMenuContext.h"
+#include "Styling/SlateIconFinder.h"
+#include "ToolMenu.h"
+#include "ToolMenuDelegates.h"
+#include "ToolMenus.h"
+#include "ToolMenuSection.h"
+#include "Widgets/Input/SNumericEntryBox.h"
+#include "Widgets/Layout/SBox.h"
+
+#define LOCTEXT_NAMESPACE "SkeletalMeshModelingToolsViewportToolbarExtension"
+
+namespace UE::SkeletalMeshModelingTools
+{
+	// Creates the skeletal mesh editor menu
+	namespace SkeletalMeshEditorMenu
+	{
+		/** Creates a slider widget */
+		TSharedRef<SWidget> CreateFloatSliderWidget(
+			const FText& TooltipText,
+			const float MaxValue,
+			const TAttribute<TOptional<float>>& ValueAttribute,
+			const SNumericEntryBox<float>::FOnValueChanged& OnValueChangedDelegate)
+		{
+			return 
+				SNew(SBox)
+				.HAlign(HAlign_Right)
+				[
+					SNew(SBox)
+					.Padding(FMargin(4.0f, 0.0f, 0.0f, 0.0f))
+					.WidthOverride(100.0f)
+					[
+						SNew(SNumericEntryBox<float>)
+						.Font(FAppStyle::GetFontStyle(TEXT("MenuItem.Font")))
+						.AllowSpin(true)
+						.MinSliderValue(0.0f)
+						.MaxSliderValue(MaxValue)
+						.Value(ValueAttribute)
+						.OnValueChanged(OnValueChangedDelegate)
+						.ToolTipText(TooltipText)
+					]
+				];			
+		}
+		
+		/** Creates a widget to edit the axis length */
+		TSharedRef<SWidget> MakeAxisLengthWidget()
+		{
+			const FText Tooltip = LOCTEXT("SkeletalMeshEditorAxesLengthTooltip", "Sets the length of displayed local rotation axes");
+			constexpr float MaxValue = 25.f;
+			
+			return SkeletalMeshEditorMenu::CreateFloatSliderWidget(
+				Tooltip,
+				MaxValue,
+				TAttribute<TOptional<float>>::CreateLambda([]()
+					{
+						return GetDefault<USkeletalMeshModelingToolsEditorSettings>()->GetLocalRotationAxisLength();
+					}),
+				SNumericEntryBox<float>::FOnValueChanged::CreateLambda([](float Value)
+					{
+						USkeletalMeshModelingToolsEditorSettings* Settings = GetMutableDefault<USkeletalMeshModelingToolsEditorSettings>();
+						Settings->SetLocalRotationAxisLength(FMath::Max(0.f, Value));
+						Settings->SaveConfig();
+					})
+				);
+		}
+		
+		/** Creates a widget to edit the axis thickness */
+		TSharedRef<SWidget> MakeAxisThicknessWidget()
+		{
+			const FText Tooltip = LOCTEXT("SkeletalMeshEditorAxesThicknessTooltip", "Sets the thickness of displayed local rotation axes");
+			constexpr float MaxValue = 1.f;
+
+			return SkeletalMeshEditorMenu::CreateFloatSliderWidget(
+				Tooltip,
+				MaxValue,
+				TAttribute<TOptional<float>>::CreateLambda([]()
+					{
+						return GetDefault<USkeletalMeshModelingToolsEditorSettings>()->GetLocalRotationAxisThickness();
+					}),
+				SNumericEntryBox<float>::FOnValueChanged::CreateLambda([](float Value)
+					{
+						USkeletalMeshModelingToolsEditorSettings* Settings = GetMutableDefault<USkeletalMeshModelingToolsEditorSettings>();
+						Settings->SetLocalRotationAxisThickness(FMath::Max(0.f, Value));
+						Settings->SaveConfig();
+					})
+				);
+		}
+
+		/** Adds the skeletal mesh menu to the toolbar */
+		void CreateSkeletalMeshEditorMenu(
+			const TSharedRef<FUICommandList>& InCommandList, 
+			UToolMenu& InOutPersonaViewportToolbar)
+		{
+			FToolMenuEntry Entry = FToolMenuEntry::InitSubMenu(
+				"SkeletalMeshEditor",
+				LOCTEXT("SkeletalMeshEditorMenuLabel", "Skeletal Mesh"),
+				LOCTEXT("SkeletalMeshEditorMenuTooltip", "Skeletal Mesh editor related settings"),
+				FNewToolMenuDelegate::CreateLambda(
+					[WeakCommandList = InCommandList.ToWeakPtr()](UToolMenu* SubMenu) -> void
+					{
+						if (const TSharedPtr<FUICommandList> CommandList = WeakCommandList.Pin())
+						{
+							FToolMenuSection& SkeletalMeshEditorSection =
+								SubMenu->FindOrAddSection(
+									"SkeletalMeshEditor",
+									LOCTEXT("SkeletalMeshEditorLabel", "Local Rotation Axes"));
+
+							const FSkeletalMeshModelingToolsCommands& Commands = FSkeletalMeshModelingToolsCommands::Get();
+
+							SkeletalMeshEditorSection.AddMenuEntryWithCommandList(Commands.ToggleShowAllLocalRotationAxes, WeakCommandList.Pin());
+
+							SkeletalMeshEditorSection.AddEntry(FToolMenuEntry::InitWidget(
+								"AxisLength",
+								MakeAxisLengthWidget(),
+								LOCTEXT("SkeletalMeshEditorLocalRotationAxesLength", "Axes Length")));
+								
+							SkeletalMeshEditorSection.AddEntry(FToolMenuEntry::InitWidget(
+								"AxisThickness",
+								MakeAxisThicknessWidget(),
+								LOCTEXT("SkeletalMeshEditorLocalRotationAxesThickness", "Axes Thickness")));
+						}
+					})
+				);
+			
+			Entry.InsertPosition = FToolMenuInsert("Show", EToolMenuInsertType::After);
+			Entry.Icon = FSlateIconFinder::FindIconForClass(USkeletalMesh::StaticClass());
+			Entry.ToolBarData.LabelOverride = FText::GetEmpty();
+
+			// Set resize prio to highest as other entries in other sections dominate otherwise
+			Entry.ToolBarData.ResizeParams.ClippingPriority = std::numeric_limits<int32>::max(); 
+
+			FToolMenuSection* const RightSection = InOutPersonaViewportToolbar.FindSection("Right");
+			if (ensureMsgf(RightSection, TEXT("Cannot find the right section in the persona viewport toolbar, cannot extend menu")))
+			{
+				RightSection->AddDynamicEntry(
+					"SkeletalMeshEditor",
+					FNewToolMenuSectionDelegate::CreateLambda([Entry](FToolMenuSection& Section)
+						{
+							// Only extending the menu for skeletal mesh editor
+							const USkeletalMeshToolMenuContext* Context = Section.FindContext<USkeletalMeshToolMenuContext>();
+							if (Context && Context->SkeletalMeshEditor.IsValid())
+							{
+								Section.AddEntry(Entry);
+							}
+						})
+				);
+			}
+		}
+	} //~ namespace SkeletalMeshEditorMenu
+
+	// Extends the transform menu
+	namespace TransformMenuExtension
+	{
+		void ExtendTransformMenu(
+			USkeletalMeshModelingToolsEditorMode& InEdMode, 
+			const TSharedRef<FUICommandList>& CommandList,
+			UToolMenu& InOutPersonaViewportToolbar)
+		{
+			const TWeakObjectPtr<USkeletalMeshModelingToolsEditorMode> WeakEdMode = &InEdMode;
+			
+			FToolMenuEntry Entry = FToolMenuEntry::InitSubMenu(
+				"Mesh Element Selection",
+				LOCTEXT("MeshElementSelectionSubmenuLabel", "Mesh Element Selection"),
+				LOCTEXT("MeshElementSelectionSubmenuTooltip", "Mesh Element Selection settings in the viewport"),
+				FNewToolMenuDelegate::CreateLambda([WeakEdMode, WeakCommandList = CommandList.ToWeakPtr()](UToolMenu* Submenu) -> void
+					{
+						FToolMenuSection& MeshElementSelectionOptionsSection =
+							Submenu->FindOrAddSection("Element Selection", LOCTEXT("ElementSelectionLabel", "Element Selection"));
+						if (const TSharedPtr<const FUICommandList>& CommandList = WeakCommandList.Pin())
+						{
+							auto CreateToolMenuEntry = [&CommandList, &MeshElementSelectionOptionsSection](const TSharedPtr<FUICommandInfo>& Command)
+							{
+								FToolMenuEntry MESMode = FToolMenuEntry::InitMenuEntryWithCommandList(Command, CommandList);
+								MESMode.SetShowInToolbarTopLevel(true); // add button to toolbar
+								MeshElementSelectionOptionsSection.AddEntry(MESMode);
+							};
+							
+							// add all Mesh Element Selection modes
+							CreateToolMenuEntry(FModelingToolsManagerCommands::Get().MeshSelectionModeAction_NoSelection);
+							MeshElementSelectionOptionsSection.AddSeparator(NAME_None);
+							CreateToolMenuEntry(FModelingToolsManagerCommands::Get().MeshSelectionModeAction_MeshVertices);
+							CreateToolMenuEntry(FModelingToolsManagerCommands::Get().MeshSelectionModeAction_MeshEdges);
+							CreateToolMenuEntry(FModelingToolsManagerCommands::Get().MeshSelectionModeAction_MeshTriangles);
+							MeshElementSelectionOptionsSection.AddSeparator(NAME_None);
+							CreateToolMenuEntry(FModelingToolsManagerCommands::Get().MeshSelectionModeAction_GroupCorners);
+							CreateToolMenuEntry(FModelingToolsManagerCommands::Get().MeshSelectionModeAction_GroupEdges);
+							CreateToolMenuEntry(FModelingToolsManagerCommands::Get().MeshSelectionModeAction_GroupFaces);
+							
+							// Section for Selection Edits
+							{
+								FToolMenuSection& SelectionEditsSection =
+								Submenu->FindOrAddSection("SelectionEdits", LOCTEXT("SelectionEditsLabel", "Selection Edits"));
+								
+								SelectionEditsSection.AddEntry(
+									FToolMenuEntry::InitMenuEntryWithCommandList(
+										FModelingToolsManagerCommands::Get().BeginSelectionAction_SelectAll, CommandList));
+								SelectionEditsSection.AddEntry(
+									FToolMenuEntry::InitMenuEntryWithCommandList(
+										FModelingToolsManagerCommands::Get().BeginSelectionAction_ExpandToConnected, CommandList));
+								SelectionEditsSection.AddEntry(
+									FToolMenuEntry::InitMenuEntryWithCommandList(
+										FModelingToolsManagerCommands::Get().BeginSelectionAction_Invert, CommandList));
+								SelectionEditsSection.AddEntry(
+									FToolMenuEntry::InitMenuEntryWithCommandList(
+										FModelingToolsManagerCommands::Get().BeginSelectionAction_InvertConnected, CommandList));
+								SelectionEditsSection.AddEntry(
+									FToolMenuEntry::InitMenuEntryWithCommandList(
+										FModelingToolsManagerCommands::Get().BeginSelectionAction_Expand, CommandList));
+								SelectionEditsSection.AddEntry(
+									FToolMenuEntry::InitMenuEntryWithCommandList(
+										FModelingToolsManagerCommands::Get().BeginSelectionAction_Contract, CommandList));
+							}
+							
+							// Section for Local Frame Mode setting
+							{
+								FToolMenuSection& LocalFrameModeSection =
+								Submenu->FindOrAddSection("Local Frame Mode", LOCTEXT("LocalFrameModeLabel", "Local Frame Mode"));
+								LocalFrameModeSection.AddEntry(
+									FToolMenuEntry::InitMenuEntryWithCommandList(
+										FModelingToolsManagerCommands::Get().SelectionLocalFrameMode_Geometry, CommandList));
+								LocalFrameModeSection.AddEntry(
+									FToolMenuEntry::InitMenuEntryWithCommandList(
+										FModelingToolsManagerCommands::Get().SelectionLocalFrameMode_Object, CommandList));
+							}
+							
+							// Section for Selection Filters
+							{
+								FToolMenuSection& SelectionFilterSection =
+								Submenu->FindOrAddSection("Selection Filters", LOCTEXT("SelectionFiltersLabel", "Selection Filters"));
+								SelectionFilterSection.AddEntry(
+									FToolMenuEntry::InitMenuEntryWithCommandList(
+										FModelingToolsManagerCommands::Get().SelectionHitBackFaces, CommandList));
+							}
+							
+							// Section for Selection Filters
+							{
+								FToolMenuSection& SoftSelectionSection =
+								Submenu->FindOrAddSection("Soft Selection", LOCTEXT("SoftSelectionLabel", "Soft Selection"));
+								SoftSelectionSection.AddEntry(
+									FToolMenuEntry::InitMenuEntryWithCommandList(
+										FModelingToolsManagerCommands::Get().SelectionEnableSoftSelection, CommandList));
+								const TSharedRef<SWidget> SoftSelectionRadiusWidget =
+									SNew(SHorizontalBox)
+									+ SHorizontalBox::Slot()
+									.FillWidth(0.9f)
+									[
+										SNew(SSpinBox<double>)
+										.MinValue(0.0f)
+										.Value_Lambda(
+											[WeakEdMode]() -> double 
+											{
+												const UGeometrySelectionManager* const SelectionManager = WeakEdMode.IsValid() ? WeakEdMode->GetSelectionManager() : nullptr;
+												return SelectionManager ? SelectionManager->GetSoftSelectionRadius() : 0.0;
+											})
+										.OnValueChanged_Lambda(
+											[WeakEdMode](double InValue)
+												{
+													UGeometrySelectionManager* const SelectionManager = WeakEdMode.IsValid() ? WeakEdMode->GetSelectionManager() : nullptr;
+													if (SelectionManager)
+													{
+														SelectionManager->SetSoftSelectionRadius(InValue);
+													}
+												})
+									]
+									+ SHorizontalBox::Slot()
+									.FillWidth(0.1f);
+								
+								SoftSelectionSection.AddEntry(
+									FToolMenuEntry::InitWidget(
+										"SoftSelectionRadius", 
+										SoftSelectionRadiusWidget, 
+										LOCTEXT("SoftSelectionRadiusLabel", "Radius")));
+							}
+						}
+					})
+			);
+
+			Entry.ToolBarData.ResizeParams.ClippingPriority = 950;
+			Entry.ToolBarData.LabelOverride = FText();
+			Entry.InsertPosition = FToolMenuInsert("Transform", EToolMenuInsertType::After);
+			Entry.Icon = FSlateIcon(FModelingToolsEditorModeStyle::GetStyleSetName(), "ModelingToolsManagerCommands.MeshElementSelection");
+
+			FToolMenuSection* const LeftSection = InOutPersonaViewportToolbar.FindSection("Left");
+			if (ensureMsgf(LeftSection, TEXT("Cannot find the left section in the persona viewport toolbar, cannot extend menu")))
+			{
+				LeftSection->AddDynamicEntry(TEXT("Mesh Element Selection"),
+					FNewToolMenuSectionDelegate::CreateLambda([Entry](FToolMenuSection& InSection)
+						{
+							// Only extending the menu for skeletal mesh editor
+							const USkeletalMeshToolMenuContext* Context = InSection.FindContext<USkeletalMeshToolMenuContext>();
+							if (Context && Context->SkeletalMeshEditor.IsValid())
+							{
+								InSection.AddEntry(Entry);
+							}
+						})
+				);
+			}
+		}
+	} //~ namespace TransformMenuExtension
+
+	static constexpr const TCHAR* SkeletalMeshEditorViewportToolbarOwnerName = TEXT("SkeletalMeshModelingToolsViewportToolbarExtension");
+
+	void ExtendSkeletalMeshEditorViewportToolbar(
+		USkeletalMeshModelingToolsEditorMode& EdMode, 
+		const TSharedRef<FUICommandList>& CommandList)
+	{
+		const FToolMenuOwnerScoped ScopeOwner(SkeletalMeshEditorViewportToolbarOwnerName);
+
+		UToolMenu* const PersonaViewportToolbar = UToolMenus::Get()->ExtendMenu("AnimationEditor.ViewportToolbar");
+		if (ensureMsgf(PersonaViewportToolbar, TEXT("Cannot find the the persona viewport toolbar, cannot extend menus")))
+		{
+			SkeletalMeshEditorMenu::CreateSkeletalMeshEditorMenu(CommandList, *PersonaViewportToolbar);
+			TransformMenuExtension::ExtendTransformMenu(EdMode, CommandList, *PersonaViewportToolbar);
+		}
+	}
+
+	void RemoveSkeletalMeshEditorViewportToolbarExtensions()
+	{
+		UToolMenus::Get()->UnregisterOwnerByName(SkeletalMeshEditorViewportToolbarOwnerName);
+	}
+}
+
+#undef LOCTEXT_NAMESPACE
diff --git a/Engine/Plugins/Animation/SkeletalMeshModelingTools/Source/SkeletalMeshModelingTools/Private/SkeletalMeshModelingToolsViewportToolbarExtension.h b/Engine/Plugins/Animation/SkeletalMeshModelingTools/Source/SkeletalMeshModelingTools/Private/SkeletalMeshModelingToolsViewportToolbarExtension.h
new file mode 100644
--- /dev/null
+++ b/Engine/Plugins/Animation/SkeletalMeshModelingTools/Source/SkeletalMeshModelingTools/Private/SkeletalMeshModelingToolsViewportToolbarExtension.h
@@ -0,0 +1,17 @@
+// Copyright Epic Games, Inc. All Rights Reserved.
+
+#pragma once
+
+#include "SkeletalMeshModelingToolsEditorModeToolkit.h"
+
+class FSkeletalMeshModelingToolsEditorModeToolkit;
+class FUICommandList;
+
+namespace UE::SkeletalMeshModelingTools
+{
+	void ExtendSkeletalMeshEditorViewportToolbar(
+		USkeletalMeshModelingToolsEditorMode& EdMode,
+		const TSharedRef<FUICommandList>& CommandList);
+
+	void RemoveSkeletalMeshEditorViewportToolbarExtensions();
+}
diff --git a/Engine/Plugins/Runtime/MeshModelingToolset/Source/MeshModelingToolsEditorOnly/Private/SkeletalMesh/SkeletonEditingTool.cpp b/Engine/Plugins/Runtime/MeshModelingToolset/Source/MeshModelingToolsEditorOnly/Private/SkeletalMesh/SkeletonEditingTool.cpp
--- a/Engine/Plugins/Runtime/MeshModelingToolset/Source/MeshModelingToolsEditorOnly/Private/SkeletalMesh/SkeletonEditingTool.cpp
+++ b/Engine/Plugins/Runtime/MeshModelingToolset/Source/MeshModelingToolsEditorOnly/Private/SkeletalMesh/SkeletonEditingTool.cpp
@@ -23,6 +23,7 @@
 
 #include "SkeletalMesh/SkeletonTransformProxy.h"
 #include "SkeletalMesh/SReferenceSkeletonTree.h"
+#include "SkeletalMesh/SkeletalMeshModelingToolsEditorSettings.h"
 #include "TargetInterfaces/PrimitiveComponentBackedTarget.h"
 #include "TargetInterfaces/SkeletonProvider.h"
 #include "TargetInterfaces/SkeletonCommitter.h"
@@ -1836,23 +1837,28 @@ void USkeletonEditingTool::Render(IToolsContextRenderAPI* RenderAPI)
 	{
 		return;
 	}
-
+	
 	const UPersonaOptions* PersonaOptions = GetDefault<UPersonaOptions>();
 	static FSkelDebugDrawConfig DrawConfig;
-		DrawConfig.BoneDrawMode = EBoneDrawMode::Type::All;
+	DrawConfig.BoneDrawMode = EBoneDrawMode::Type::All;
 #if WITH_EDITORONLY_DATA
-		DrawConfig.BoneDrawSize = WeakMesh.IsValid() ? WeakMesh->GetBoneDrawSize() : 1.f;
+	DrawConfig.BoneDrawSize = WeakMesh.IsValid() ? WeakMesh->GetBoneDrawSize() : 1.f;
 #else
-		DrawConfig.BoneDrawSize = 1.f;
+	DrawConfig.BoneDrawSize = 1.f;
 #endif
-		DrawConfig.bAddHitProxy = true;
-		DrawConfig.bForceDraw = false;
-		DrawConfig.DefaultBoneColor = PersonaOptions->DefaultBoneColor;
-		DrawConfig.AffectedBoneColor = PersonaOptions->AffectedBoneColor;
-		DrawConfig.SelectedBoneColor = PersonaOptions->SelectedBoneColor;
-		DrawConfig.ParentOfSelectedBoneColor = PersonaOptions->ParentOfSelectedBoneColor;
-		DrawConfig.AxisConfig.Thickness = Properties->AxisThickness;
-		DrawConfig.AxisConfig.Length = Properties->AxisLength;
+	DrawConfig.bAddHitProxy = true;
+	DrawConfig.bForceDraw = false;
+	DrawConfig.DefaultBoneColor = PersonaOptions->DefaultBoneColor;
+	DrawConfig.AffectedBoneColor = PersonaOptions->AffectedBoneColor;
+	DrawConfig.SelectedBoneColor = PersonaOptions->SelectedBoneColor;
+	DrawConfig.ParentOfSelectedBoneColor = PersonaOptions->ParentOfSelectedBoneColor;
+	
+	if (const USkeletalMeshModelingToolsEditorSettings* const EditorSettings = GetDefault<USkeletalMeshModelingToolsEditorSettings>())
+	{
+		DrawConfig.AxisConfig.Length = EditorSettings->GetLocalRotationAxisLength();
+		DrawConfig.AxisConfig.Thickness = EditorSettings->GetLocalRotationAxisThickness();
+		DrawConfig.bDrawAllLocalRotationAxes = EditorSettings->GetShowAllLocalRotationAxes();
+	}
 	
 	FPrimitiveDrawInterface* PDI = RenderAPI->GetPrimitiveDrawInterface();
 
diff --git a/Engine/Plugins/Runtime/MeshModelingToolset/Source/MeshModelingToolsEditorOnly/Private/SkeletalMesh/SkinWeightsPaintTool.cpp b/Engine/Plugins/Runtime/MeshModelingToolset/Source/MeshModelingToolsEditorOnly/Private/SkeletalMesh/SkinWeightsPaintTool.cpp
--- a/Engine/Plugins/Runtime/MeshModelingToolset/Source/MeshModelingToolsEditorOnly/Private/SkeletalMesh/SkinWeightsPaintTool.cpp
+++ b/Engine/Plugins/Runtime/MeshModelingToolset/Source/MeshModelingToolsEditorOnly/Private/SkeletalMesh/SkinWeightsPaintTool.cpp
@@ -2005,6 +2005,29 @@ void USkinWeightsPaintTool::Setup()
 		EToolMessageLevel::UserNotification);
 }
 
+void USkinWeightsPaintTool::Shutdown(EToolShutdownType ShutdownType) 
+{
+	SkeletonTreeNotifierBindScope.Reset();
+	
+	ISkeletalMeshBackedTarget* SkeletalMeshBackedTarget = Cast<ISkeletalMeshBackedTarget>(GetTarget()); 
+	USkeletalMesh* TargetSkeletalMesh = SkeletalMeshBackedTarget ? SkeletalMeshBackedTarget->GetSkeletalMesh() : nullptr;
+	if (TargetSkeletalMesh)
+	{
+		TargetSkeletalMesh->Modify();
+	}
+	
+	if (ShutdownType == EToolShutdownType::Accept)
+	{
+		// to force to refresh the tree
+		if (NeedsNotification())
+		{
+			GetNotifier()->Notify({}, ESkeletalMeshNotifyType::BonesAdded);
+		}
+	}
+	
+	Super::Shutdown(ShutdownType);
+}
+
 void USkinWeightsPaintTool::DrawHUD(FCanvas* Canvas, IToolsContextRenderAPI* RenderAPI)
 {
 	Super::DrawHUD(Canvas, RenderAPI);
diff --git a/Engine/Plugins/Runtime/MeshModelingToolset/Source/MeshModelingToolsEditorOnly/Public/SkeletalMesh/SkeletalMeshModelingToolsEditorSettings.h b/Engine/Plugins/Runtime/MeshModelingToolset/Source/MeshModelingToolsEditorOnly/Public/SkeletalMesh/SkeletalMeshModelingToolsEditorSettings.h
new file mode 100644
--- /dev/null
+++ b/Engine/Plugins/Runtime/MeshModelingToolset/Source/MeshModelingToolsEditorOnly/Public/SkeletalMesh/SkeletalMeshModelingToolsEditorSettings.h
@@ -0,0 +1,47 @@
+// Copyright Epic Games, Inc. All Rights Reserved.
+
+#pragma once
+
+#include "UObject/Object.h"
+
+#include "SkeletalMeshModelingToolsEditorSettings.generated.h"
+
+#define UE_API MESHMODELINGTOOLSEDITORONLY_API
+
+UCLASS(config = EditorPerProjectUserSettings)
+class USkeletalMeshModelingToolsEditorSettings 
+	: public UObject
+{
+	GENERATED_BODY()
+
+public:
+	/** Sets the local rotation axis length */
+	UE_API void SetLocalRotationAxisLength(const float InLocalRotationAxisLength) { LocalRotationAxisLength = InLocalRotationAxisLength; }
+
+	/** Gets the local rotation axis length */
+	UE_API float GetLocalRotationAxisLength() const { return LocalRotationAxisLength; }
+
+	/** Sets the local rotation axis thickness */
+	UE_API void SetLocalRotationAxisThickness(const float InLocalRotationAxisThickness) { LocalRotationAxisThickness = InLocalRotationAxisThickness; }
+
+	/** Gets the local rotation axis thickness */
+	UE_API float GetLocalRotationAxisThickness() const { return LocalRotationAxisThickness; }
+
+	/** Sets if all or only the selected local rotation axes are displayed */
+	UE_API void SetShowAllLocalRotationAxes(const bool InShowAllLocalRotationAxes) { bShowAllLocalRotationAxes = InShowAllLocalRotationAxes; }
+
+	/** Gets if all or only the selected local rotation axes are displayed */
+	UE_API bool GetShowAllLocalRotationAxes() const { return bShowAllLocalRotationAxes; }
+
+private:
+	UPROPERTY(config)
+	float LocalRotationAxisLength = 0.f;
+	
+	UPROPERTY(config)
+	float LocalRotationAxisThickness = 0.f;
+
+	UPROPERTY(config)
+	bool bShowAllLocalRotationAxes = false;
+};
+
+#undef UE_API
diff --git a/Engine/Plugins/Runtime/MeshModelingToolset/Source/MeshModelingToolsEditorOnly/Public/SkeletalMesh/SkeletonEditingTool.h b/Engine/Plugins/Runtime/MeshModelingToolset/Source/MeshModelingToolsEditorOnly/Public/SkeletalMesh/SkeletonEditingTool.h
--- a/Engine/Plugins/Runtime/MeshModelingToolset/Source/MeshModelingToolsEditorOnly/Public/SkeletalMesh/SkeletonEditingTool.h
+++ b/Engine/Plugins/Runtime/MeshModelingToolset/Source/MeshModelingToolsEditorOnly/Public/SkeletalMesh/SkeletonEditingTool.h
@@ -352,12 +352,6 @@ public:
 	UPROPERTY(EditAnywhere, Category = "Details")
 	bool bUpdateChildren = false;
 
-	UPROPERTY(EditAnywhere, Category = "Viewport Axis Settings",  meta = (DisplayPriority = 10, ClampMin = "0.0", UIMin = "0.0"))
-	float AxisLength = 1.f;
-
-	UPROPERTY(EditAnywhere, Category = "Viewport Axis Settings",  meta = (DisplayPriority = 10, ClampMin = "0.0", UIMin = "0.0"))
-	float AxisThickness = 0.f;
-
 	UPROPERTY()
 	bool bEnableComponentSelection = false;
 
diff --git a/Engine/Plugins/Runtime/MeshModelingToolset/Source/MeshModelingToolsEditorOnly/Public/SkeletalMesh/SkinWeightsPaintTool.h b/Engine/Plugins/Runtime/MeshModelingToolset/Source/MeshModelingToolsEditorOnly/Public/SkeletalMesh/SkinWeightsPaintTool.h
--- a/Engine/Plugins/Runtime/MeshModelingToolset/Source/MeshModelingToolsEditorOnly/Public/SkeletalMesh/SkinWeightsPaintTool.h
+++ b/Engine/Plugins/Runtime/MeshModelingToolset/Source/MeshModelingToolsEditorOnly/Public/SkeletalMesh/SkinWeightsPaintTool.h
@@ -667,6 +667,7 @@ public:
 	
 	// UInteractiveTool
 	UE_API virtual void Setup() override;
+	UE_API virtual void Shutdown(EToolShutdownType ShutdownType) override;
 	UE_API virtual void DrawHUD(FCanvas* Canvas, IToolsContextRenderAPI* RenderAPI) override;
 	UE_API virtual void Render(IToolsContextRenderAPI* RenderAPI) override;
 
diff --git a/Engine/Source/Runtime/Engine/Private/SkeletalDebugRendering.cpp b/Engine/Source/Runtime/Engine/Private/SkeletalDebugRendering.cpp
--- a/Engine/Source/Runtime/Engine/Private/SkeletalDebugRendering.cpp
+++ b/Engine/Source/Runtime/Engine/Private/SkeletalDebugRendering.cpp
@@ -405,8 +405,8 @@ void DrawBonesInternal(
 		FLinearColor BoneColor = bIsAffected ? DrawConfig.AffectedBoneColor : DefaultBoneColor;
 		BoneColor = bIsSelected ? DrawConfig.SelectedBoneColor : BoneColor;
 
-		// draw the little coordinate frame inside the bone ONLY if selected or affected
-		AxisConfig.bDraw = bIsAffected || bIsSelected;
+		// Draw the local rotation axes for affected and selected bones, or for all related option is set to true
+		AxisConfig.bDraw = DrawConfig.bDrawAllLocalRotationAxes || bIsAffected || bIsSelected;
 
 		// draw cone to each child
 		// but use a different color if this bone is NOT selected, but the child IS selected
diff --git a/Engine/Source/Runtime/Engine/Public/SkeletalDebugRendering.h b/Engine/Source/Runtime/Engine/Public/SkeletalDebugRendering.h
--- a/Engine/Source/Runtime/Engine/Public/SkeletalDebugRendering.h
+++ b/Engine/Source/Runtime/Engine/Public/SkeletalDebugRendering.h
@@ -46,6 +46,12 @@ public:
 	FLinearColor SelectedBoneColor;
 	FLinearColor ParentOfSelectedBoneColor;
 	FBoneAxisDrawConfig AxisConfig;
+	
+	/** 
+	 * For methods that draw multiple bones, draws local rotation axes for all when set to true.
+	 * For methods that draw a single bone this option has no effect.
+	 */
+	bool bDrawAllLocalRotationAxes = false;
 };
 
 namespace SkeletalDebugRendering
