[SKMModelingTools & Persona]
SKMModelingTools: Fixed crash when toggling editing tools while a tool is active
Persona: Added viewport interaction support for skeleton selection edit mode

#Synced-CL 51144723

diff --git a/Engine/Plugins/Animation/SkeletalMeshModelingTools/Source/SkeletalMeshModelingTools/Private/SkeletalMeshEditingCache.cpp b/Engine/Plugins/Animation/SkeletalMeshModelingTools/Source/SkeletalMeshModelingTools/Private/SkeletalMeshEditingCache.cpp
--- a/Engine/Plugins/Animation/SkeletalMeshModelingTools/Source/SkeletalMeshModelingTools/Private/SkeletalMeshEditingCache.cpp
+++ b/Engine/Plugins/Animation/SkeletalMeshModelingTools/Source/SkeletalMeshModelingTools/Private/SkeletalMeshEditingCache.cpp
@@ -275,6 +275,7 @@ void USkeletalMeshEditingCache::Destroy()
 	if (SkeletalMeshComponent.IsValid())
 	{
 		SkeletalMeshComponent->SetVisibility(true);
+		SetSkeletalMeshSkeletonDrawMode(ESkeletonDrawMode::Default);
 	}
 }
 
diff --git a/Engine/Plugins/Animation/SkeletalMeshModelingTools/Source/SkeletalMeshModelingTools/Private/SkeletalMeshModelingToolsEditorMode.cpp b/Engine/Plugins/Animation/SkeletalMeshModelingTools/Source/SkeletalMeshModelingTools/Private/SkeletalMeshModelingToolsEditorMode.cpp
--- a/Engine/Plugins/Animation/SkeletalMeshModelingTools/Source/SkeletalMeshModelingTools/Private/SkeletalMeshModelingToolsEditorMode.cpp
+++ b/Engine/Plugins/Animation/SkeletalMeshModelingTools/Source/SkeletalMeshModelingTools/Private/SkeletalMeshModelingToolsEditorMode.cpp
@@ -14,9 +14,6 @@
 #include "EditorViewportClient.h"
 #include "Features/IModularFeatures.h"
 
-// Stylus support is currently disabled due to issues with the stylus plugin
-// We are leaving the code in this cpp file, defined out, so that it is easier to bring back if/when the stylus plugin is improved.
-
 #include "SkeletalMeshModelingToolsModule.h"
 #include "Selection/SKMGeometrySelectionTransformTweaker.h"
 #include "ViewportInteractions/SKMModelingToolsClickSelection.h"
@@ -28,8 +25,13 @@
 #include "ViewportInteractions/ViewportMoveYawInteraction.h"
 #include "ViewportInteractions/ViewportOrbitInteraction.h"
 #include "Widgets/Notifications/SNotificationList.h"
+#include "Misc/MessageDialog.h"
 
 #include UE_INLINE_GENERATED_CPP_BY_NAME(SkeletalMeshModelingToolsEditorMode)
+
+// Stylus support is currently disabled due to issues with the stylus plugin
+// We are leaving the code in this cpp file, defined out, so that it is easier to bring back if/when the stylus plugin is improved.
+
 #define ENABLE_STYLUS_SUPPORT 0
 
 #if ENABLE_STYLUS_SUPPORT 
@@ -721,9 +723,29 @@ void USkeletalMeshModelingToolsEditorMode::Exit()
 		SelectionManager = nullptr;
 	}
 
-	if (GetCurrentEditingCache())
+	
+	if (GetCurrentEditingCache() )
 	{
-		GetCurrentEditingCache()->ApplyChanges();
+		// Inform users if there are unapplied changes
+		int32 ChangeCount = GetCurrentEditingCache()->GetEditingMeshComponent()->GetChangeCount();
+		if (ChangeCount > 0 && GetSkeletalMesh())
+		{
+			EAppReturnType::Type YesNoReply = FMessageDialog::Open(EAppMsgType::YesNo,
+					FText::Format(
+						LOCTEXT("Prompt_SkeletalMeshModelingToolsClose", "Skeletal Mesh \"{Mesh}\" has {ChangeCount} unapplied changes. Would you like to apply them? (Selecting 'No' will cause all changes to be lost!)"),
+						FText::FromString(GetSkeletalMesh()->GetPathName()),
+						ChangeCount));
+
+			if (YesNoReply == EAppReturnType::Yes)
+			{
+				GetCurrentEditingCache()->ApplyChanges();
+			}
+			else
+			{
+				GetCurrentEditingCache()->DiscardChanges();
+			}
+		}
+
 		GetCurrentEditingCache()->Destroy();
 		CurrentEditingCache = nullptr;
 	}
@@ -1053,7 +1075,9 @@ void USkeletalMeshModelingToolsEditorMode::InitializeGeometrySelectionSystems()
 
 void USkeletalMeshModelingToolsEditorMode::OnBuildViewportInteractions(UViewportInteractionsBehaviorSource* Source)
 {
-	Source->AddInteraction<USkeletalMeshModelingToolsClickSelection>();
+	USkeletalMeshModelingToolsClickSelection* BaseClickSelection = 
+		Source->AddInteraction<USkeletalMeshModelingToolsClickSelection>();
+	BaseClickSelection->Bind(this);
 
 	USkeletalMeshModelingToolsGeometryFrustumSelection* GeometryFrustumSelection =
 		Source->AddInteraction<USkeletalMeshModelingToolsGeometryFrustumSelection>();
@@ -1136,6 +1160,9 @@ void USkeletalMeshModelingToolsEditorMode::Tick(FEditorViewportClient* ViewportC
 {
 	Super::Tick(ViewportClient, DeltaTime);
 
+	GeometrySelectionStateWatcher.CheckAndUpdate();
+	SkeletonStateUpdater.CheckAndUpdate();
+
 	if (GetCurrentEditingCache())
 	{
 		GetCurrentEditingCache()->Tick();
@@ -1144,8 +1171,6 @@ void USkeletalMeshModelingToolsEditorMode::Tick(FEditorViewportClient* ViewportC
 	CacheChangeCountWatcher.CheckAndUpdate();
 	SkeletonTreeHostStateUpdater.CheckAndUpdate();
 	CacheSkeletonChangeCountWatcher.CheckAndUpdate();
-	GeometrySelectionStateWatcher.CheckAndUpdate();
-	SkeletonStateUpdater.CheckAndUpdate();
 }
 
 bool USkeletalMeshModelingToolsEditorMode::HandleClick(FEditorViewportClient* InViewportClient, HHitProxy *HitProxy, const FViewportClick &Click)
diff --git a/Engine/Plugins/Animation/SkeletalMeshModelingTools/Source/SkeletalMeshModelingTools/Private/SkeletalMeshModelingToolsEditorModeToolkit.cpp b/Engine/Plugins/Animation/SkeletalMeshModelingTools/Source/SkeletalMeshModelingTools/Private/SkeletalMeshModelingToolsEditorModeToolkit.cpp
--- a/Engine/Plugins/Animation/SkeletalMeshModelingTools/Source/SkeletalMeshModelingTools/Private/SkeletalMeshModelingToolsEditorModeToolkit.cpp
+++ b/Engine/Plugins/Animation/SkeletalMeshModelingTools/Source/SkeletalMeshModelingTools/Private/SkeletalMeshModelingToolsEditorModeToolkit.cpp
@@ -231,15 +231,16 @@ void FSkeletalMeshModelingToolsEditorModeToolkit::InvokeUI()
 
 void FSkeletalMeshModelingToolsEditorModeToolkit::ShutdownUI()
 {
-	//Skeleton
-	DynamicMeshSkeletonTreeNotifierBindScope.Reset();
+	// Restore Skeleton UI
+	EnableSkeletalMeshSkeletonCommands();
+	ShowSkeletalMeshSkeletonTree();
 	
+	// Restore Morph Target	UI
 	if (ModeUILayer.IsValid())
 	{
 		TSharedPtr<FAssetEditorModeUILayer> ModeUILayerPtr = ModeUILayer.Pin();
 		TSharedPtr<FTabManager> TabManager = ModeUILayerPtr->GetTabManager();
 
-		// Morph Target
 		{
 			const FName& MorphTargetTabID = FPersonaTabs::MorphTargetsID;
 			TSharedPtr<SDockTab> MorphTargetTab = TabManager->FindExistingLiveTab(MorphTargetTabID);
diff --git a/Engine/Plugins/Animation/SkeletalMeshModelingTools/Source/SkeletalMeshModelingTools/Private/ViewportInteractions/SKMModelingToolsClickSelection.cpp b/Engine/Plugins/Animation/SkeletalMeshModelingTools/Source/SkeletalMeshModelingTools/Private/ViewportInteractions/SKMModelingToolsClickSelection.cpp
--- a/Engine/Plugins/Animation/SkeletalMeshModelingTools/Source/SkeletalMeshModelingTools/Private/ViewportInteractions/SKMModelingToolsClickSelection.cpp
+++ b/Engine/Plugins/Animation/SkeletalMeshModelingTools/Source/SkeletalMeshModelingTools/Private/ViewportInteractions/SKMModelingToolsClickSelection.cpp
@@ -4,6 +4,7 @@
 
 #include "EditorModeManager.h"
 #include "EditorViewportClient.h"
+#include "SkeletalMeshModelingToolsEditorMode.h"
 
 USkeletalMeshModelingToolsClickSelection::USkeletalMeshModelingToolsClickSelection()
 {
@@ -11,6 +12,11 @@ USkeletalMeshModelingToolsClickSelection::USkeletalMeshModelingToolsClickSelecti
 	Groups = { UE::Editor::ViewportInteractions::ViewportClick };
 }
 
+void USkeletalMeshModelingToolsClickSelection::Bind(USkeletalMeshModelingToolsEditorMode* InEditorMode)
+{
+	EditorMode = InEditorMode;
+}
+
 FInputRayHit USkeletalMeshModelingToolsClickSelection::IsHitByClick(const FInputDeviceRay& ClickPos)
 {
 	if (IsEnabled())
@@ -23,5 +29,5 @@ FInputRayHit USkeletalMeshModelingToolsClickSelection::IsHitByClick(const FInput
 
 void USkeletalMeshModelingToolsClickSelection::ProcessClick_Internal(const FViewportClick& InViewportClick, HHitProxy* InHitProxy, FSceneView& View)
 {
-	GetEditorViewportClient()->GetModeTools()->HandleClick(GetEditorViewportClient(),  InHitProxy, InViewportClick);
+	EditorMode->HandleClick(GetEditorViewportClient(),  InHitProxy, InViewportClick);
 }
diff --git a/Engine/Plugins/Animation/SkeletalMeshModelingTools/Source/SkeletalMeshModelingTools/Private/ViewportInteractions/SKMModelingToolsClickSelection.h b/Engine/Plugins/Animation/SkeletalMeshModelingTools/Source/SkeletalMeshModelingTools/Private/ViewportInteractions/SKMModelingToolsClickSelection.h
--- a/Engine/Plugins/Animation/SkeletalMeshModelingTools/Source/SkeletalMeshModelingTools/Private/ViewportInteractions/SKMModelingToolsClickSelection.h
+++ b/Engine/Plugins/Animation/SkeletalMeshModelingTools/Source/SkeletalMeshModelingTools/Private/ViewportInteractions/SKMModelingToolsClickSelection.h
@@ -6,6 +6,8 @@
 #include "SKMModelingToolsClickSelection.generated.h"
 
 
+class USkeletalMeshModelingToolsEditorMode;
+
 UCLASS(MinimalAPI, Transient)
 class USkeletalMeshModelingToolsClickSelection
 	: public UViewportClickInteraction
@@ -14,9 +16,12 @@ class USkeletalMeshModelingToolsClickSelection
 
 public:
 	USkeletalMeshModelingToolsClickSelection();
+	void Bind(USkeletalMeshModelingToolsEditorMode* InEditorMode);
 	virtual FInputRayHit IsHitByClick(const FInputDeviceRay& ClickPos) override;
 
 protected:
 	virtual void ProcessClick_Internal(const FViewportClick& InViewportClick, HHitProxy* InHitProxy, FSceneView& View) override;
+
+	TWeakObjectPtr<USkeletalMeshModelingToolsEditorMode> EditorMode;
 };
 
diff --git a/Engine/Source/Editor/Persona/Persona.Build.cs b/Engine/Source/Editor/Persona/Persona.Build.cs
--- a/Engine/Source/Editor/Persona/Persona.Build.cs
+++ b/Engine/Source/Editor/Persona/Persona.Build.cs
@@ -20,6 +20,7 @@ public class Persona : ModuleRules
             new string[] {
                 "AdvancedPreviewScene",
                 "AnimationEditMode",
+				"EditorInteractiveToolsFramework"
             }
         );
 
diff --git a/Engine/Source/Editor/Persona/Private/SkeletonSelectionEditMode.cpp b/Engine/Source/Editor/Persona/Private/SkeletonSelectionEditMode.cpp
--- a/Engine/Source/Editor/Persona/Private/SkeletonSelectionEditMode.cpp
+++ b/Engine/Source/Editor/Persona/Private/SkeletonSelectionEditMode.cpp
@@ -17,6 +17,9 @@
 #include "IEditableSkeleton.h"
 
 #include "SceneView.h"
+#include "Tools/EdModeInteractiveToolsContext.h"
+#include "ViewportInteractions/SkeletonSelectionClickInteraction.h"
+#include "ViewportInteractions/ViewportInteractionsBehaviorSource.h"
 
 #define LOCTEXT_NAMESPACE "SkeletonSelectionEditMode"
 
@@ -41,6 +44,16 @@ const FReferenceSkeleton& FSkeletonSelectionEditMode::GetReferenceSkeletonForCom
 	return Component->GetSkeletalMeshAsset() && !Component->GetSkeletalMeshAsset()->IsCompiling() ? Component->GetSkeletalMeshAsset()->GetRefSkeleton() : GetAnimPreviewScene().GetPersonaToolkit()->GetSkeleton()->GetReferenceSkeleton();
 }
 
+void FSkeletonSelectionEditMode::OnBuildViewportInteractions(UViewportInteractionsBehaviorSource* Source)
+{
+	USkeletonSelectionClickInteraction* ClickInteraction = Source->AddInteraction<USkeletonSelectionClickInteraction>();
+	ClickInteraction->Bind(SharedThis(this));
+
+	USkeletonSelectionRightClickInteraction* RightClickInteraction = Source->AddInteraction<USkeletonSelectionRightClickInteraction>();
+	RightClickInteraction->Bind(SharedThis(this));
+}
+
+
 bool FSkeletonSelectionEditMode::GetCameraTarget(FSphere& OutTarget) const
 {
 	bool bHandled = false;
@@ -86,6 +99,26 @@ void FSkeletonSelectionEditMode::GetOnScreenDebugInfo(TArray<FText>& OutDebugInf
 
 }
 
+void FSkeletonSelectionEditMode::Enter()
+{
+	IPersonaEditMode::Enter();
+
+	if (UModeManagerInteractiveToolsContext* const ModeManagerInteractiveToolsContext = GetModeManager()->GetInteractiveToolsContext())
+	{
+		ModeManagerInteractiveToolsContext->OnBuildViewportInteractions().AddSP(this, &FSkeletonSelectionEditMode::OnBuildViewportInteractions);
+	}	
+}
+
+void FSkeletonSelectionEditMode::Exit()
+{
+	if (UModeManagerInteractiveToolsContext* const ModeManagerInteractiveToolsContext = GetModeManager()->GetInteractiveToolsContext())
+	{
+		ModeManagerInteractiveToolsContext->OnBuildViewportInteractions().RemoveAll(this);
+	}
+	
+	IPersonaEditMode::Exit();
+}
+
 FSelectedSocketInfo FSkeletonSelectionEditMode::DuplicateAndSelectSocket(const FSelectedSocketInfo& SocketInfoToDuplicate)
 {
 	USkeletalMesh* SkeletalMesh = GetAnimPreviewScene().GetPreviewMeshComponent()->GetSkeletalMeshAsset();
diff --git a/Engine/Source/Editor/Persona/Private/ViewportInteractions/SkeletonSelectionClickInteraction.cpp b/Engine/Source/Editor/Persona/Private/ViewportInteractions/SkeletonSelectionClickInteraction.cpp
new file mode 100644
--- /dev/null
+++ b/Engine/Source/Editor/Persona/Private/ViewportInteractions/SkeletonSelectionClickInteraction.cpp
@@ -0,0 +1,41 @@
+// Copyright Epic Games, Inc. All Rights Reserved.
+
+#include "ViewportInteractions/SkeletonSelectionClickInteraction.h"
+#include "SkeletonSelectionEditMode.h"
+#include "BaseBehaviors/SingleClickBehavior.h"
+
+USkeletonSelectionClickInteraction::USkeletonSelectionClickInteraction()
+{
+	InteractionName = TEXT("Skeleton Selection Click Interaction");
+	Groups = { UE::Editor::ViewportInteractions::ViewportClick };
+}
+
+void USkeletonSelectionClickInteraction::Bind(TSharedPtr<FSkeletonSelectionEditMode> InEditMode)
+{
+	EditMode = InEditMode;
+}
+
+FInputRayHit USkeletonSelectionClickInteraction::IsHitByClick(const FInputDeviceRay& ClickPos)
+{
+	if (IsEnabled())
+	{
+		return Super::IsHitByClick(ClickPos);
+	}
+	
+	return FInputRayHit();
+}
+
+void USkeletonSelectionClickInteraction::ProcessClick_Internal(const FViewportClick& InViewportClick, HHitProxy* InHitProxy, FSceneView& View)
+{
+	EditMode.Pin()->HandleClick(GetEditorViewportClient(),  InHitProxy, InViewportClick);
+}
+
+USkeletonSelectionRightClickInteraction::USkeletonSelectionRightClickInteraction()
+{
+	InteractionName = TEXT("Skeleton Selection Right Click Interaction");
+	Groups = { UE::Editor::ViewportInteractions::ViewportClick };
+
+	MouseButton = EKeys::RightMouseButton;
+	
+	SingleClickInputBehaviorWeak->SetUseRightMouseButton();
+}
diff --git a/Engine/Source/Editor/Persona/Public/SkeletonSelectionEditMode.h b/Engine/Source/Editor/Persona/Public/SkeletonSelectionEditMode.h
--- a/Engine/Source/Editor/Persona/Public/SkeletonSelectionEditMode.h
+++ b/Engine/Source/Editor/Persona/Public/SkeletonSelectionEditMode.h
@@ -18,6 +18,7 @@ struct FSelectedSocketInfo;
 class FViewportClient;
 struct FReferenceSkeleton;
 struct FGizmoState;
+class UViewportInteractionsBehaviorSource;
 
 class FSkeletonSelectionEditMode : public IPersonaEditMode
 {
@@ -30,6 +31,8 @@ public:
 	virtual void GetOnScreenDebugInfo(TArray<FText>& OutDebugInfo) const override;
 
 	/** FEdMode interface */
+	virtual void Enter() override;
+	virtual void Exit() override;
 	virtual bool StartTracking(FEditorViewportClient* InViewportClient, FViewport* InViewport) override;
 	virtual bool EndTracking(FEditorViewportClient* InViewportClient, FViewport* InViewport) override;
 	virtual bool InputDelta(FEditorViewportClient* InViewportClient, FViewport* InViewport, FVector& InDrag, FRotator& InRot, FVector& InScale) override;
@@ -53,7 +56,8 @@ protected:
 	FTransform GetSocketTransform(const USkeletalMeshSocket* Socket) const;
 	FTransform GetWorldSpaceBoneTransform(const FReferenceSkeleton& ReferenceSkeleton, const int32 BoneIndex) const;
 	const FReferenceSkeleton& GetReferenceSkeletonForComponent(const USkeletalMeshComponent* Component) const;
-
+	void OnBuildViewportInteractions(UViewportInteractionsBehaviorSource* Source);
+	
 private:
 	/** Duplicates and selects a socket when we alt-drag */
 	FSelectedSocketInfo DuplicateAndSelectSocket(const FSelectedSocketInfo& SocketInfoToDuplicate);
diff --git a/Engine/Source/Editor/Persona/Public/ViewportInteractions/SkeletonSelectionClickInteraction.h b/Engine/Source/Editor/Persona/Public/ViewportInteractions/SkeletonSelectionClickInteraction.h
new file mode 100644
--- /dev/null
+++ b/Engine/Source/Editor/Persona/Public/ViewportInteractions/SkeletonSelectionClickInteraction.h
@@ -0,0 +1,38 @@
+// Copyright Epic Games, Inc. All Rights Reserved.
+
+#pragma once
+
+#include "ViewportInteractions/ViewportClickInteraction.h"
+#include "SkeletonSelectionClickInteraction.generated.h"
+
+class FSkeletonSelectionEditMode;
+
+UCLASS(MinimalAPI, Transient)
+class USkeletonSelectionClickInteraction
+	: public UViewportClickInteraction
+{
+	GENERATED_BODY()
+
+public:
+	USkeletonSelectionClickInteraction();
+	void Bind(TSharedPtr<FSkeletonSelectionEditMode> InEditMode);
+	
+	virtual FInputRayHit IsHitByClick(const FInputDeviceRay& ClickPos) override;
+
+protected:
+	virtual void ProcessClick_Internal(const FViewportClick& InViewportClick, HHitProxy* InHitProxy, FSceneView& View) override;
+
+	TWeakPtr<FSkeletonSelectionEditMode> EditMode;
+};
+
+
+UCLASS(MinimalAPI, Transient)
+class USkeletonSelectionRightClickInteraction
+	: public USkeletonSelectionClickInteraction
+{
+	GENERATED_BODY()
+
+public:
+	USkeletonSelectionRightClickInteraction();
+	
+};
