[SKMModellingTools] Fixed click selection not working for geometry selection

#Synced-CL 51254598

diff --git a/Engine/Plugins/Animation/SkeletalMeshModelingTools/Source/SkeletalMeshModelingTools/Private/ViewportInteractions/SKMModelingToolsGeometryClickSelection.cpp b/Engine/Plugins/Animation/SkeletalMeshModelingTools/Source/SkeletalMeshModelingTools/Private/ViewportInteractions/SKMModelingToolsGeometryClickSelection.cpp
--- a/Engine/Plugins/Animation/SkeletalMeshModelingTools/Source/SkeletalMeshModelingTools/Private/ViewportInteractions/SKMModelingToolsGeometryClickSelection.cpp
+++ b/Engine/Plugins/Animation/SkeletalMeshModelingTools/Source/SkeletalMeshModelingTools/Private/ViewportInteractions/SKMModelingToolsGeometryClickSelection.cpp
@@ -6,32 +6,53 @@
 #include "EditorViewportClient.h"
 #include "BaseBehaviors/SingleClickBehavior.h"
 #include "Selection/GeometrySelectionManager.h"
+#include "ViewportInteractions/ViewportInteractionsBehaviorSource.h"
 
 USkeletalMeshModelingToolsGeometryClickSelection::USkeletalMeshModelingToolsGeometryClickSelection()
 {
 	InteractionName = TEXT("Skeletal Mesh Modeling Tools Geometry Click Selection");
 	Groups = { UE::Editor::ViewportInteractions::ViewportClick };
-
-	USingleClickInputBehavior* SingleClickBehavior = NewObject<USingleClickInputBehavior>();
-	SingleClickBehavior->Initialize(this);
-
-	SingleClickBehavior->SetDefaultPriority(UE::Editor::ViewportInteractions::CLICK_PRIORITY);
-	SingleClickBehavior->SetAllowsCaptureStealing(true);
-	RegisterInputBehavior(SingleClickBehavior);
 }
 
-FInputRayHit USkeletalMeshModelingToolsGeometryClickSelection::IsHitByClick(const FInputDeviceRay& ClickPos)
+void USkeletalMeshModelingToolsGeometryClickSelection::BuildBehaviors()
 {
-	if (IsEnabled())
+	Super::BuildBehaviors();
+	
+	if (UViewportClickBehavior* ClickBehavior = ClickBehaviorWeak.Get())
 	{
-		return FInputRayHit((double)TNumericLimits<float>::Max());
-	}
+		using namespace UE::Editor::ViewportInteractions;
 	
-	return FInputRayHit();
+		ClickBehavior->SetBindings({
+			FButtonBinding(EKeys::LeftMouseButton).TriggersStart(true),
+			FButtonBinding(EKeys::LeftControl).Required(false),
+			FButtonBinding(EKeys::LeftShift).Required(false)
+		});
+	}
 }
 
-void USkeletalMeshModelingToolsGeometryClickSelection::OnClicked(const FInputDeviceRay& ClickPos)
+void USkeletalMeshModelingToolsGeometryClickSelection::OnClickUp(const FInputDeviceRay& InClickPos)
 {
+	if (UViewportInteractionsBehaviorSource* const ViewportInteractionBehaviorSource = GetViewportInteractionsBehaviorSource())
+	{
+		ViewportInteractionBehaviorSource->ClearMouseCursorOverride();
+	}
+
+	if (IsMouseLooking())
+	{
+		if (UViewportInteractionsBehaviorSource* const ViewportInteractionBehaviorSource = GetViewportInteractionsBehaviorSource())
+		{
+			// Store info about camera movement locally, since SetIsMouseLooking call will reset it
+			const bool bCameraHasMoved = ViewportInteractionBehaviorSource->HasCameraMoved();
+			ViewportInteractionBehaviorSource->SetIsMouseLooking(false);
+
+			// If camera has moved, this was definitely not a click
+			if (bCameraHasMoved)
+			{
+				return;
+			}
+		}
+	}
+
 	UE::Geometry::FGeometrySelectionUpdateConfig UpdateConfig;
 	UpdateConfig.ChangeType = UE::Geometry::EGeometrySelectionChangeType::Replace;
 	if (IsShiftDown())
@@ -43,7 +64,7 @@ void USkeletalMeshModelingToolsGeometryClickSelection::OnClicked(const FInputDev
 		UpdateConfig.ChangeType = UE::Geometry::EGeometrySelectionChangeType::Remove;
 	}
 	UE::Geometry::FGeometrySelectionUpdateResult Result;
-	SelectionManager->UpdateSelectionViaRaycast(ClickPos.WorldRay, UpdateConfig, Result);	
+	SelectionManager->UpdateSelectionViaRaycast(InClickPos.WorldRay, UpdateConfig, Result);	
 }
 
 void USkeletalMeshModelingToolsGeometryClickSelection::BindSelectionManager(UGeometrySelectionManager* InSelectionManager)
diff --git a/Engine/Plugins/Animation/SkeletalMeshModelingTools/Source/SkeletalMeshModelingTools/Private/ViewportInteractions/SKMModelingToolsGeometryClickSelection.h b/Engine/Plugins/Animation/SkeletalMeshModelingTools/Source/SkeletalMeshModelingTools/Private/ViewportInteractions/SKMModelingToolsGeometryClickSelection.h
--- a/Engine/Plugins/Animation/SkeletalMeshModelingTools/Source/SkeletalMeshModelingTools/Private/ViewportInteractions/SKMModelingToolsGeometryClickSelection.h
+++ b/Engine/Plugins/Animation/SkeletalMeshModelingTools/Source/SkeletalMeshModelingTools/Private/ViewportInteractions/SKMModelingToolsGeometryClickSelection.h
@@ -9,15 +9,14 @@ class UGeometrySelectionManager;
 
 UCLASS(MinimalAPI, Transient)
 class USkeletalMeshModelingToolsGeometryClickSelection
-	: public UViewportInteraction
-	, public IClickBehaviorTarget
+	: public UViewportClickInteraction
 {
 	GENERATED_BODY()
 
 public:
 	USkeletalMeshModelingToolsGeometryClickSelection();
-	virtual FInputRayHit IsHitByClick(const FInputDeviceRay& ClickPos) override;
-	virtual void OnClicked(const FInputDeviceRay& ClickPos) override;
+	virtual void BuildBehaviors() override;
+	virtual void OnClickUp(const FInputDeviceRay& InClickPos) override;
 
 	void BindSelectionManager(UGeometrySelectionManager* InSelectionManager);
 
diff --git a/Engine/Plugins/Animation/SkeletalMeshModelingTools/Source/SkeletalMeshModelingTools/Private/ViewportInteractions/SKMModelingToolsGeometryFrustumSelection.cpp b/Engine/Plugins/Animation/SkeletalMeshModelingTools/Source/SkeletalMeshModelingTools/Private/ViewportInteractions/SKMModelingToolsGeometryFrustumSelection.cpp
--- a/Engine/Plugins/Animation/SkeletalMeshModelingTools/Source/SkeletalMeshModelingTools/Private/ViewportInteractions/SKMModelingToolsGeometryFrustumSelection.cpp
+++ b/Engine/Plugins/Animation/SkeletalMeshModelingTools/Source/SkeletalMeshModelingTools/Private/ViewportInteractions/SKMModelingToolsGeometryFrustumSelection.cpp
@@ -15,8 +15,11 @@ USkeletalMeshModelingToolsGeometryFrustumSelection::USkeletalMeshModelingToolsGe
 
 	if (UViewportClickDragBehavior* ClickDragBehavior = ClickDragInputBehavior.Get())
 	{
+		using namespace UE::Editor::ViewportInteractions;
 		ClickDragBehavior->SetBindings({
-			UE::Editor::ViewportInteractions::FButtonBinding(EKeys::LeftMouseButton).TriggersStart()
+			FButtonBinding(EKeys::LeftMouseButton).TriggersStart(),
+			FButtonBinding(EKeys::LeftControl).Required(false),
+			FButtonBinding(EKeys::LeftShift).Required(false)
 		});
 	}
 }
diff --git a/Engine/Source/Editor/Experimental/EditorInteractiveToolsFramework/Public/ViewportInteractions/ViewportInteractionsBehaviorSource.h b/Engine/Source/Editor/Experimental/EditorInteractiveToolsFramework/Public/ViewportInteractions/ViewportInteractionsBehaviorSource.h
--- a/Engine/Source/Editor/Experimental/EditorInteractiveToolsFramework/Public/ViewportInteractions/ViewportInteractionsBehaviorSource.h
+++ b/Engine/Source/Editor/Experimental/EditorInteractiveToolsFramework/Public/ViewportInteractions/ViewportInteractionsBehaviorSource.h
@@ -210,7 +210,7 @@ public:
 	 */
 	EDITORINTERACTIVETOOLSFRAMEWORK_API void SetMouseCursorOverride(EMouseCursor::Type InMouseCursor);
 
-	void ClearMouseCursorOverride();
+	EDITORINTERACTIVETOOLSFRAMEWORK_API void ClearMouseCursorOverride();
 
 	/**
 	 * Use the templated version in case there is some class specific initialization required by the newly created interaction,
@@ -254,7 +254,7 @@ public:
 
 	bool IsMouseLooking() const;
 
-	void SetIsMouseLooking(bool bInIsLooking);
+	EDITORINTERACTIVETOOLSFRAMEWORK_API void SetIsMouseLooking(bool bInIsLooking);
 
 	bool HasCameraMoved() const
 	{
