[SkeletalMeshEditingTools] Improve the Isolate Selection button to look like a button even when toggled off

#rnx
#Synced-CL 51338101

diff --git a/Engine/Plugins/Animation/SkeletalMeshModelingTools/Source/SkeletalMeshModelingTools/Private/DetailCustomization/SkinWeightDetailCustomization.cpp b/Engine/Plugins/Animation/SkeletalMeshModelingTools/Source/SkeletalMeshModelingTools/Private/DetailCustomization/SkinWeightDetailCustomization.cpp
--- a/Engine/Plugins/Animation/SkeletalMeshModelingTools/Source/SkeletalMeshModelingTools/Private/DetailCustomization/SkinWeightDetailCustomization.cpp
+++ b/Engine/Plugins/Animation/SkeletalMeshModelingTools/Source/SkeletalMeshModelingTools/Private/DetailCustomization/SkinWeightDetailCustomization.cpp
@@ -13,6 +13,7 @@
 #include "Widgets/Input/SNumericEntryBox.h"
 #include "Framework/MultiBox/MultiBoxBuilder.h"
 #include "SkeletalMesh/SReferenceSkeletonTree.h"
+#include "SkeletalMeshModelingToolsStyle.h"
 
 #define LOCTEXT_NAMESPACE "SkinWeightToolSettingsEditor"
 
@@ -445,7 +446,7 @@ void FSkinWeightDetailCustomization::AddSelectionUI(IDetailLayoutBuilder& Detail
 				.VAlign(VAlign_Center)
 				[
 					SNew(SCheckBox)
-					.Style(FAppStyle::Get(), "ToggleButtonCheckbox")
+					.Style(FSkeletalMeshModelingToolsStyle::Get(), "SkeletalMeshModelingTools.ToggleButtonCheckbox")
 					.HAlign(HAlign_Center)
 					.ToolTipText(LOCTEXT("IsolateSelectedTooltip",
 							"Shows only the selected faces in the viewport.\n"
diff --git a/Engine/Plugins/Animation/SkeletalMeshModelingTools/Source/SkeletalMeshModelingTools/Private/SkeletalMeshModelingToolsStyle.cpp b/Engine/Plugins/Animation/SkeletalMeshModelingTools/Source/SkeletalMeshModelingTools/Private/SkeletalMeshModelingToolsStyle.cpp
--- a/Engine/Plugins/Animation/SkeletalMeshModelingTools/Source/SkeletalMeshModelingTools/Private/SkeletalMeshModelingToolsStyle.cpp
+++ b/Engine/Plugins/Animation/SkeletalMeshModelingTools/Source/SkeletalMeshModelingTools/Private/SkeletalMeshModelingToolsStyle.cpp
@@ -2,30 +2,39 @@
 
 #include "SkeletalMeshModelingToolsStyle.h"
 
+#include "Styling/AppStyle.h"
 #include "Styling/SlateStyleMacros.h"
 #include "Styling/SlateStyleRegistry.h"
+#include "Styling/SlateTypes.h"
+#include "Styling/StyleColors.h"
 
-
-FSkeletalMeshModelingToolsStyle::FSkeletalMeshModelingToolsStyle() : 
-	FSlateStyleSet("SkeletalMeshModelingToolsStyle")
+FSkeletalMeshModelingToolsStyle::FSkeletalMeshModelingToolsStyle() 
+	: FSlateStyleSet("SkeletalMeshModelingToolsStyle")
 {
-	
+	FCheckBoxStyle BaseCheckBoxStyle = FAppStyle::Get().GetWidgetStyle<FCheckBoxStyle>("ToggleButtonCheckbox");
+	Set("SkeletalMeshModelingTools.ToggleButtonCheckbox", 
+		BaseCheckBoxStyle
+		.SetUncheckedImage(FSlateRoundedBoxBrush(FStyleColors::Header, 4.0f, FStyleColors::Input, 1.0f))
+		.SetUncheckedHoveredImage(FSlateRoundedBoxBrush(FStyleColors::Hover, 4.0f, FStyleColors::Input, 1.0f))
+		.SetUncheckedPressedImage(FSlateRoundedBoxBrush(FStyleColors::Hover, 4.0f, FStyleColors::Input, 1.0f))
+		.SetCheckedImage(FSlateRoundedBoxBrush(FStyleColors::Primary, 4.0f, FStyleColors::Input, 1.0f))
+		.SetCheckedHoveredImage(FSlateRoundedBoxBrush(FStyleColors::PrimaryHover, 4.0f, FStyleColors::Input, 1.0f))
+		.SetCheckedPressedImage(FSlateRoundedBoxBrush(FStyleColors::PrimaryHover, 4.0f, FStyleColors::Input, 1.0f))
+		.SetPadding(FMargin(16, 6))
+	);
 }
 
-
 FSkeletalMeshModelingToolsStyle& FSkeletalMeshModelingToolsStyle::Get()
 {
 	static FSkeletalMeshModelingToolsStyle Instance;
 	return Instance;
 }
 
-
 void FSkeletalMeshModelingToolsStyle::Register()
 {
 	FSlateStyleRegistry::RegisterSlateStyle(Get());
 }
 
-
 void FSkeletalMeshModelingToolsStyle::Unregister()
 {
 	FSlateStyleRegistry::UnRegisterSlateStyle(Get());
