[SKMModelingTools] gate new attribute tool using the same cvar that modeling mode uses

#Synced-CL 51489786

diff --git a/Engine/Plugins/Animation/SkeletalMeshModelingTools/Source/SkeletalMeshModelingTools/Private/SkeletalMeshModelingToolsEditorMode.cpp b/Engine/Plugins/Animation/SkeletalMeshModelingTools/Source/SkeletalMeshModelingTools/Private/SkeletalMeshModelingToolsEditorMode.cpp
--- a/Engine/Plugins/Animation/SkeletalMeshModelingTools/Source/SkeletalMeshModelingTools/Private/SkeletalMeshModelingToolsEditorMode.cpp
+++ b/Engine/Plugins/Animation/SkeletalMeshModelingTools/Source/SkeletalMeshModelingTools/Private/SkeletalMeshModelingToolsEditorMode.cpp
@@ -26,6 +26,7 @@
 #include "ViewportInteractions/ViewportOrbitInteraction.h"
 #include "Widgets/Notifications/SNotificationList.h"
 #include "Misc/MessageDialog.h"
+#include "HAL/ConsoleManager.h"
 
 #include UE_INLINE_GENERATED_CPP_BY_NAME(SkeletalMeshModelingToolsEditorMode)
 
@@ -359,7 +360,22 @@ void USkeletalMeshModelingToolsEditorMode::Enter()
 	RegisterTool(ToolManagerCommands.BeginDisplaceMeshTool, TEXT("BeginDisplaceMeshTool"), NewObject<UDisplaceMeshToolBuilder>());
 
 	RegisterTool(ToolManagerCommands.BeginAttributeEditorTool, TEXT("BeginAttributeEditorTool"), NewObject<UAttributeEditorToolBuilder>());
-	RegisterTool(ToolManagerCommands.BeginMeshAttributePaintTool, TEXT("BeginMeshAttributePaintTool"), NewObject<UMeshAttributePaintToolV2Builder>());
+	bool bUseExperimentalAttributePaintTool = true;
+	if (IConsoleVariable* Variable = IConsoleManager::Get().FindConsoleVariable(TEXT("ModelingTools.UseExperimentalMeshAttributePaintTool")))
+	{
+		if (Variable->GetBool() == false)
+		{
+			bUseExperimentalAttributePaintTool = false;
+		}
+	}
+	if (bUseExperimentalAttributePaintTool)
+	{
+		RegisterTool(ToolManagerCommands.BeginMeshAttributePaintTool, TEXT("BeginMeshAttributePaintTool"), NewObject<UMeshAttributePaintToolV2Builder>());
+	}
+	else
+	{
+		RegisterTool(ToolManagerCommands.BeginMeshAttributePaintTool, TEXT("BeginMeshAttributePaintTool"), NewObject<UMeshAttributePaintToolBuilder>());
+	}
 	RegisterTool(ToolManagerCommands.BeginMeshVertexPaintTool, TEXT("BeginMeshVertexPaintTool"), NewObject<UMeshVertexPaintToolBuilder>());
 	RegisterTool(ToolManagerCommands.BeginSkinWeightsPaintTool, TEXT("BeginSkinWeightsPaintTool"), NewObject<USkinWeightsPaintToolBuilder>());
 	RegisterTool(ToolManagerCommands.BeginSkinWeightsBindingTool, TEXT("BeginSkinWeightsBindingTool"), NewObject<USkinWeightsBindingToolBuilder>());
