[SKMModelingTools] minor fixes for morph manager tab, added a hotkey to toggle soft selection

#Synced-CL 51489786

diff --git a/Engine/Plugins/Animation/SkeletalMeshModelingTools/Source/SkeletalMeshModelingTools/Private/SkeletalMeshModelingToolsEditorMode.cpp b/Engine/Plugins/Animation/SkeletalMeshModelingTools/Source/SkeletalMeshModelingTools/Private/SkeletalMeshModelingToolsEditorMode.cpp
--- a/Engine/Plugins/Animation/SkeletalMeshModelingTools/Source/SkeletalMeshModelingTools/Private/SkeletalMeshModelingToolsEditorMode.cpp
+++ b/Engine/Plugins/Animation/SkeletalMeshModelingTools/Source/SkeletalMeshModelingTools/Private/SkeletalMeshModelingToolsEditorMode.cpp
@@ -676,12 +676,36 @@ FName USkeletalMeshModelingToolsEditorMode::HandleAddMorphTarget(FName InName)
 
 FName USkeletalMeshModelingToolsEditorMode::HandleRenameMorphTarget(FName OldName, FName NewName)
 {
-	return GetCurrentEditingCache()->RenameMorphTarget(OldName, NewName);
+	bool bSetEditingMorphTargetAgain = false;
+	if (GetEditingMorphTarget() == OldName)
+	{
+		bSetEditingMorphTargetAgain = true;
+	}
+	
+	FName FinalNewName = GetCurrentEditingCache()->RenameMorphTarget(OldName, NewName);
+
+	if (bSetEditingMorphTargetAgain)
+	{
+		HandleSetEditingMorphTarget(FinalNewName);
+	}
+	
+	return FinalNewName;
 }
 
 void USkeletalMeshModelingToolsEditorMode::HandleRemoveMorphTargets(const TArray<FName>& Names)
 {
+	bool bClearEditingMorphTarget= false;
+	if (Names.Contains(GetEditingMorphTarget()))
+	{
+		bClearEditingMorphTarget = true;
+	}
+	
 	GetCurrentEditingCache()->RemoveMorphTargets(Names);
+
+	if (bClearEditingMorphTarget)
+	{
+		HandleSetEditingMorphTarget(NAME_None);
+	}
 }
 
 TArray<FName> USkeletalMeshModelingToolsEditorMode::HandleDuplicateMorphTargets(const TArray<FName>& Names)
@@ -999,7 +1023,14 @@ void USkeletalMeshModelingToolsEditorMode::InitializeGeometrySelectionSystems()
 			bool bIsSoftSelectionEnabled = WeakThis->SelectionManager->IsSoftSelectionEnabled();
 			WeakThis->SelectionManager->SetSoftSelectionEnabled(!bIsSoftSelectionEnabled);
 		}),
-		FCanExecuteAction(),
+		FCanExecuteAction::CreateLambda([WeakThis]()
+		{
+			if (!WeakThis.IsValid())
+			{
+				return false;
+			}
+			return WeakThis->IsGeometrySelectionActive();	
+		}),
 		FIsActionChecked::CreateLambda([WeakThis]()
 		{
 			if (!WeakThis.IsValid())
diff --git a/Engine/Plugins/Editor/ModelingToolsEditorMode/Source/ModelingToolsEditorMode/Private/ModelingToolsManagerActions.cpp b/Engine/Plugins/Editor/ModelingToolsEditorMode/Source/ModelingToolsEditorMode/Private/ModelingToolsManagerActions.cpp
--- a/Engine/Plugins/Editor/ModelingToolsEditorMode/Source/ModelingToolsEditorMode/Private/ModelingToolsManagerActions.cpp
+++ b/Engine/Plugins/Editor/ModelingToolsEditorMode/Source/ModelingToolsEditorMode/Private/ModelingToolsManagerActions.cpp
@@ -266,8 +266,8 @@ void FModelingToolsManagerCommands::RegisterCommands()
 
 	REGISTER_MODELING_TOOL_COMMAND(SelectionHitBackFaces, "Hit Back Faces", "Hit triangles that are facing away from the camera");
 	REGISTER_MODELING_TOOL_COMMAND(SelectionLocking, "Lock Target", "Click to Toggle locking on the Selected Object and disallow/allow Mesh Selections");
-	
-	REGISTER_MODELING_TOOL_COMMAND(SelectionEnableSoftSelection, "Soft Selection", "Apply a falloff effect to nearby vertices when transforming the selected elements");
+
+	UI_COMMAND(SelectionEnableSoftSelection, "Soft Selection", "Apply a falloff effect to nearby vertices when transforming the selected elements", EUserInterfaceActionType::ToggleButton, FInputChord(EKeys::B));
 	
 	UI_COMMAND(AddToFavorites, "Add to Favorites", "Add to Favorites", EUserInterfaceActionType::Button, FInputChord());
 	UI_COMMAND(RemoveFromFavorites, "Remove from Favorites", "Remove from Favorites", EUserInterfaceActionType::Button, FInputChord());
