[SkeletalMeshModelingTools] Add MirrorEditingMorphTarget menu entry to selection menu in SKM tools

#Synced-CL 51489786

diff --git a/Engine/Plugins/Animation/SkeletalMeshModelingTools/Source/SkeletalMeshModelingTools/Private/SkeletalMeshModelingToolsViewportToolbarExtension.cpp b/Engine/Plugins/Animation/SkeletalMeshModelingTools/Source/SkeletalMeshModelingTools/Private/SkeletalMeshModelingToolsViewportToolbarExtension.cpp
--- a/Engine/Plugins/Animation/SkeletalMeshModelingTools/Source/SkeletalMeshModelingTools/Private/SkeletalMeshModelingToolsViewportToolbarExtension.cpp
+++ b/Engine/Plugins/Animation/SkeletalMeshModelingTools/Source/SkeletalMeshModelingTools/Private/SkeletalMeshModelingToolsViewportToolbarExtension.cpp
@@ -282,6 +282,15 @@ namespace UE::SkeletalMeshModelingTools
 										SoftSelectionRadiusWidget, 
 										LOCTEXT("SoftSelectionRadiusLabel", "Radius")));
 							}
+
+							// Section for Morph Target actions
+							{
+								FToolMenuSection& MorphTargetSection =
+								Submenu->FindOrAddSection("Morph Target", LOCTEXT("MorphTargetLabel", "Morph Target"));
+								MorphTargetSection.AddEntry(
+									FToolMenuEntry::InitMenuEntryWithCommandList(
+										FSkeletalMeshModelingToolsCommands::Get().MirrorEditingMorphTarget, CommandList));
+							}
 						}
 					})
 			);
