[SKM Modeling Tools] added FlipMorphTarget for skeletal mesh modeling tools

#Synced-CL 51489786

diff --git a/Engine/Plugins/Animation/SkeletalMeshModelingTools/Source/SkeletalMeshModelingTools/Private/Components/SKMBackedDynaMeshComponent.cpp b/Engine/Plugins/Animation/SkeletalMeshModelingTools/Source/SkeletalMeshModelingTools/Private/Components/SKMBackedDynaMeshComponent.cpp
--- a/Engine/Plugins/Animation/SkeletalMeshModelingTools/Source/SkeletalMeshModelingTools/Private/Components/SKMBackedDynaMeshComponent.cpp
+++ b/Engine/Plugins/Animation/SkeletalMeshModelingTools/Source/SkeletalMeshModelingTools/Private/Components/SKMBackedDynaMeshComponent.cpp
@@ -666,7 +666,102 @@ void USkeletalMeshBackedDynamicMeshComponent::MirrorMorphTargets(const TArray<FN
 		Trackers.MorphTargetChangeTracker.HandleEditMorphTarget(Name);
 	}
 
-	MarkDirtyInternal();	
+	MarkDirtyInternal();
+}
+
+void USkeletalMeshBackedDynamicMeshComponent::FlipMorphTargets(const TArray<FName>& InNames)
+{
+	UpdateSymmetryIfNeeded();
+	check(ReadOnlySymmetry);
+
+	TArray<FName> MorphsToEdit;
+	for (const FName& Name : InNames)
+	{
+		if (GetMorphTargetChangeTracker().GetCurrentMorphTargetNames().Contains(Name))
+		{
+			MorphsToEdit.Add(Name);
+		}
+	}
+
+	if (MorphsToEdit.IsEmpty())
+	{
+		return;
+	}
+
+	FChangeScope ChangeScope(this);
+
+	EditMesh([&](FDynamicMesh3& Mesh)
+	{
+		using namespace UE::Geometry;
+		for (const FName& Name : MorphsToEdit)
+		{
+			FDynamicMeshMorphTargetAttribute* SourceMorphTargetAttribute = Mesh.Attributes()->GetMorphTargetAttribute(Name);
+
+			ParallelFor(Mesh.MaxVertexID(), [&](int32 VertID)
+			{
+				if (!Mesh.IsVertex(VertID))
+				{
+					return;
+				}
+
+				FVector Position = Mesh.GetVertex(VertID);
+
+				FVector Delta;
+				SourceMorphTargetAttribute->GetValue(VertID, Delta);
+
+				constexpr bool bForceSameSizeWithGaps = true;
+				TArray<int32> MirroredVertID;
+				ReadOnlySymmetry->GetMirrorVertexROI({VertID}, MirroredVertID, bForceSameSizeWithGaps);
+
+				// INDEX_NONE means either an on-plane vert or an unmatched vert with no mirror pair.
+				// Constrain (Position + Delta) to the symmetry plane: if the constraint changes the delta,
+				// the delta has a symmetry-axis component, indicating the vert is on-plane — flip it.
+				// If unchanged, the vert is either unmatched or has no axis component to flip — skip.
+				if (MirroredVertID[0] == INDEX_NONE)
+				{
+					TArray<FVector> ConstrainedPosition;
+					ConstrainedPosition.Add(Position + Delta);
+					ReadOnlySymmetry->ApplySymmetryPlaneConstraints({VertID}, ConstrainedPosition);
+
+					FVector ConstrainedDelta = ConstrainedPosition[0] - Position;
+
+					if (!ConstrainedDelta.Equals(Delta))
+					{
+						FVector FlippedDelta = ReadOnlySymmetry->GetMirroredAxis(Delta);
+						SourceMorphTargetAttribute->SetValue(VertID, FlippedDelta);
+					}
+				}
+				else
+				{
+					// Process positive side only; read both deltas, then swap with mirrored axis
+					if (Position.X < 0)
+					{
+						return;
+					}
+
+					int32 NegVertID = MirroredVertID[0];
+
+					FVector NegDelta;
+					SourceMorphTargetAttribute->GetValue(NegVertID, NegDelta);
+
+					FVector MirroredNegDelta = ReadOnlySymmetry->GetMirroredAxis(NegDelta);
+					FVector MirroredPosDelta = ReadOnlySymmetry->GetMirroredAxis(Delta);
+
+					// Swap: pos_vert gets mirrored neg delta, neg_vert gets mirrored pos delta
+					SourceMorphTargetAttribute->SetValue(VertID, MirroredNegDelta);
+					SourceMorphTargetAttribute->SetValue(NegVertID, MirroredPosDelta);
+				}
+			});
+		}
+	});
+
+	// Tracker change
+	for (const FName& Name : MorphsToEdit)
+	{
+		Trackers.MorphTargetChangeTracker.HandleEditMorphTarget(Name);
+	}
+
+	MarkDirtyInternal();
 }
 
 void USkeletalMeshBackedDynamicMeshComponent::MarkMorphTargetEdited(FName InName)
diff --git a/Engine/Plugins/Animation/SkeletalMeshModelingTools/Source/SkeletalMeshModelingTools/Private/Components/SKMBackedDynaMeshComponent.h b/Engine/Plugins/Animation/SkeletalMeshModelingTools/Source/SkeletalMeshModelingTools/Private/Components/SKMBackedDynaMeshComponent.h
--- a/Engine/Plugins/Animation/SkeletalMeshModelingTools/Source/SkeletalMeshModelingTools/Private/Components/SKMBackedDynaMeshComponent.h
+++ b/Engine/Plugins/Animation/SkeletalMeshModelingTools/Source/SkeletalMeshModelingTools/Private/Components/SKMBackedDynaMeshComponent.h
@@ -93,7 +93,8 @@ public:
 	void RemoveMorphTargets(const TArray<FName>& InNames);
 	TArray<FName> DuplicateMorphTargets(const TArray<FName>& InNames);
 	void MirrorMorphTargets(const TArray<FName>& InNames);
-	
+	void FlipMorphTargets(const TArray<FName>& InNames);
+
 	void MarkMorphTargetEdited(FName InName);
 	
 	const FSkeletonChangeTracker& GetSkeletonChangeTracker() const {return Trackers.SkeletonChangeTracker;}
diff --git a/Engine/Plugins/Animation/SkeletalMeshModelingTools/Source/SkeletalMeshModelingTools/Private/SMorphTargetManager.cpp b/Engine/Plugins/Animation/SkeletalMeshModelingTools/Source/SkeletalMeshModelingTools/Private/SMorphTargetManager.cpp
--- a/Engine/Plugins/Animation/SkeletalMeshModelingTools/Source/SkeletalMeshModelingTools/Private/SMorphTargetManager.cpp
+++ b/Engine/Plugins/Animation/SkeletalMeshModelingTools/Source/SkeletalMeshModelingTools/Private/SMorphTargetManager.cpp
@@ -358,6 +358,11 @@ void SMorphTargetManager::BindCommands()
 		Commands.MirrorSelectedMorphTarget,
 		FExecuteAction::CreateSP(this, &SMorphTargetManager::MirrorSelectedMorphTarget),
 		FCanExecuteAction::CreateSP(this, &SMorphTargetManager::CanMirror));
+
+	CommandList->MapAction(
+		Commands.FlipSelectedMorphTarget,
+		FExecuteAction::CreateSP(this, &SMorphTargetManager::FlipSelectedMorphTarget),
+		FCanExecuteAction::CreateSP(this, &SMorphTargetManager::CanFlip));
 }
 
 void SMorphTargetManager::RefreshList()
@@ -474,7 +479,8 @@ TSharedPtr<SWidget> SMorphTargetManager::OnGetContextMenuContent()
 	MenuBuilder.AddMenuEntry(FGenericCommands::Get().Delete);
 	MenuBuilder.AddMenuEntry(FGenericCommands::Get().Duplicate);
 	MenuBuilder.AddMenuEntry(FSkeletalMeshModelingToolsCommands::Get().MirrorSelectedMorphTarget);
-	
+	MenuBuilder.AddMenuEntry(FSkeletalMeshModelingToolsCommands::Get().FlipSelectedMorphTarget);
+
 	MenuBuilder.EndSection();
 	return MenuBuilder.MakeWidget();
 }
@@ -603,6 +609,25 @@ bool SMorphTargetManager::CanMirror()
 	return ListView->GetSelectedItems().Num() == 1;
 }
 
+void SMorphTargetManager::FlipSelectedMorphTarget()
+{
+	TArray<FMorphTargetInfoPtr> Items = ListView->GetSelectedItems();
+
+	check(Items.Num() == 1);
+
+	FName MorphTargetName = Items[0]->Name;
+	FText Title = FText::Format( LOCTEXT("FlipMorphTransaction", "Flip {0}"), FText::FromName(MorphTargetName));
+	FScopedTransaction Transaction(Title);
+
+	SetEditingMorphTarget(MorphTargetName);
+	Delegates.OnFlipSelectedMorphTarget.Execute(MorphTargetName);
+}
+
+bool SMorphTargetManager::CanFlip()
+{
+	return ListView->GetSelectedItems().Num() == 1;
+}
+
 
 TSharedRef<class ITableRow> SMorphTargetManager::GenerateMorphTargetRow(
 	FMorphTargetInfoPtr MorphTargetItem,
diff --git a/Engine/Plugins/Animation/SkeletalMeshModelingTools/Source/SkeletalMeshModelingTools/Private/SMorphTargetManager.h b/Engine/Plugins/Animation/SkeletalMeshModelingTools/Source/SkeletalMeshModelingTools/Private/SMorphTargetManager.h
--- a/Engine/Plugins/Animation/SkeletalMeshModelingTools/Source/SkeletalMeshModelingTools/Private/SMorphTargetManager.h
+++ b/Engine/Plugins/Animation/SkeletalMeshModelingTools/Source/SkeletalMeshModelingTools/Private/SMorphTargetManager.h
@@ -37,6 +37,7 @@ public:
 	DECLARE_DELEGATE_OneParam(FOnRemoveMorphTargets, const TArray<FName>&);
 	DECLARE_DELEGATE_RetVal_OneParam(TArray<FName>, FOnDuplicateMorphTargets, const TArray<FName>&);
 	DECLARE_DELEGATE_OneParam(FOnMirrorMorphTarget, FName);
+	DECLARE_DELEGATE_OneParam(FOnFlipMorphTarget, FName);
 
 	struct FMorphTargetManagerDelegates
 	{
@@ -44,18 +45,19 @@ public:
 
 		FOnGetMorphTargetWeight OnGetMorphTargetWeight;
 		FOnSetMorphTargetWeight OnSetMorphTargetWeight;
-		
+
 		FOnGetMorphTargetAutoFill OnGetMorphTargetAutoFill;
 		FOnSetMorphTargetAutoFill OnSetMorphTargetAutoFill;
 
 		FOnGetEditingMorphTarget OnGetEditingMorphTarget;
 		FOnSetEditingMorphTarget OnSetEditingMorphTarget;
-		
+
 		FOnAddMorphTarget OnAddNewMorphTarget;
 		FOnRenameMorphTarget OnRenameMorphTarget;
 		FOnRemoveMorphTargets OnRemoveMorphTargets;
 		FOnDuplicateMorphTargets OnDuplicateMorphTargets;
 		FOnMirrorMorphTarget OnMirrorSelectedMorphTarget;
+		FOnFlipMorphTarget OnFlipSelectedMorphTarget;
 	};	
 	
 	SLATE_BEGIN_ARGS( SMorphTargetManager ) {}
@@ -103,7 +105,9 @@ protected:
 	bool CanDuplicate();
 	void MirrorSelectedMorphTarget();
 	bool CanMirror();
-	
+	void FlipSelectedMorphTarget();
+	bool CanFlip();
+
 	TSharedPtr<SSearchBox> NameFilterBox;
 	TSharedPtr<SMorphTargetManagerListType> ListView;
 
diff --git a/Engine/Plugins/Animation/SkeletalMeshModelingTools/Source/SkeletalMeshModelingTools/Private/SkeletalMeshEditingCache.cpp b/Engine/Plugins/Animation/SkeletalMeshModelingTools/Source/SkeletalMeshModelingTools/Private/SkeletalMeshEditingCache.cpp
--- a/Engine/Plugins/Animation/SkeletalMeshModelingTools/Source/SkeletalMeshModelingTools/Private/SkeletalMeshEditingCache.cpp
+++ b/Engine/Plugins/Animation/SkeletalMeshModelingTools/Source/SkeletalMeshModelingTools/Private/SkeletalMeshEditingCache.cpp
@@ -751,6 +751,11 @@ void USkeletalMeshEditingCache::MirrorMorphTargets(const TArray<FName>& InNames)
 	EditingMeshComponent->MirrorMorphTargets(InNames);
 }
 
+void USkeletalMeshEditingCache::FlipMorphTargets(const TArray<FName>& InNames)
+{
+	EditingMeshComponent->FlipMorphTargets(InNames);
+}
+
 void USkeletalMeshEditingCache::HandleGeometryUpdate(const FDynamicMesh3& InMesh, FName MorphTargetName)
 {
 	using namespace SkeletalMeshToolsHelper;
diff --git a/Engine/Plugins/Animation/SkeletalMeshModelingTools/Source/SkeletalMeshModelingTools/Private/SkeletalMeshEditingCache.h b/Engine/Plugins/Animation/SkeletalMeshModelingTools/Source/SkeletalMeshModelingTools/Private/SkeletalMeshEditingCache.h
--- a/Engine/Plugins/Animation/SkeletalMeshModelingTools/Source/SkeletalMeshModelingTools/Private/SkeletalMeshEditingCache.h
+++ b/Engine/Plugins/Animation/SkeletalMeshModelingTools/Source/SkeletalMeshModelingTools/Private/SkeletalMeshEditingCache.h
@@ -102,6 +102,7 @@ public:
 	void RemoveMorphTargets(const TArray<FName>& InNames);
 	TArray<FName> DuplicateMorphTargets(const TArray<FName>& InNames);
 	void MirrorMorphTargets(const TArray<FName>& InNames);
+	void FlipMorphTargets(const TArray<FName>& InNames);
 
 	void HandleGeometryUpdate(const UE::Geometry::FDynamicMesh3& InMesh, FName MorphTargetName);
 	
diff --git a/Engine/Plugins/Animation/SkeletalMeshModelingTools/Source/SkeletalMeshModelingTools/Private/SkeletalMeshModelingToolsCommands.cpp b/Engine/Plugins/Animation/SkeletalMeshModelingTools/Source/SkeletalMeshModelingTools/Private/SkeletalMeshModelingToolsCommands.cpp
--- a/Engine/Plugins/Animation/SkeletalMeshModelingTools/Source/SkeletalMeshModelingTools/Private/SkeletalMeshModelingToolsCommands.cpp
+++ b/Engine/Plugins/Animation/SkeletalMeshModelingTools/Source/SkeletalMeshModelingTools/Private/SkeletalMeshModelingToolsCommands.cpp
@@ -18,6 +18,10 @@ void FSkeletalMeshModelingToolsCommands::RegisterCommands()
 	
 	UI_COMMAND(MirrorEditingMorphTarget, "Mirror Editing Morph Target", "Mirror the Morph Target currently selected for editing, if geometry selection is active", EUserInterfaceActionType::Button, FInputChord(EKeys::M));
 
+	UI_COMMAND(FlipSelectedMorphTarget, "Flip Selected Morph Target", "Flip the selected Morph Target", EUserInterfaceActionType::Button, FInputChord());
+
+	UI_COMMAND(FlipEditingMorphTarget, "Flip Editing Morph Target", "Flip the Morph Target currently selected for editing, if geometry selection is active", EUserInterfaceActionType::Button, FInputChord());
+
 	UI_COMMAND(ToggleShowAllLocalRotationAxes, "Show All Local Rotation Axes", "When enabled shows all local rotation axes, otherwise for the selected bone only", EUserInterfaceActionType::ToggleButton, FInputChord());	
 }
 
diff --git a/Engine/Plugins/Animation/SkeletalMeshModelingTools/Source/SkeletalMeshModelingTools/Private/SkeletalMeshModelingToolsCommands.h b/Engine/Plugins/Animation/SkeletalMeshModelingTools/Source/SkeletalMeshModelingTools/Private/SkeletalMeshModelingToolsCommands.h
--- a/Engine/Plugins/Animation/SkeletalMeshModelingTools/Source/SkeletalMeshModelingTools/Private/SkeletalMeshModelingToolsCommands.h
+++ b/Engine/Plugins/Animation/SkeletalMeshModelingTools/Source/SkeletalMeshModelingTools/Private/SkeletalMeshModelingToolsCommands.h
@@ -33,6 +33,10 @@ public:
 	
 	TSharedPtr<FUICommandInfo> MirrorEditingMorphTarget;
 
+	TSharedPtr<FUICommandInfo> FlipSelectedMorphTarget;
+
+	TSharedPtr<FUICommandInfo> FlipEditingMorphTarget;
+
     /** Command to toggle show all local rotation axes */
 	TSharedPtr<FUICommandInfo> ToggleShowAllLocalRotationAxes;
 };
diff --git a/Engine/Plugins/Animation/SkeletalMeshModelingTools/Source/SkeletalMeshModelingTools/Private/SkeletalMeshModelingToolsEditorMode.cpp b/Engine/Plugins/Animation/SkeletalMeshModelingTools/Source/SkeletalMeshModelingTools/Private/SkeletalMeshModelingToolsEditorMode.cpp
--- a/Engine/Plugins/Animation/SkeletalMeshModelingTools/Source/SkeletalMeshModelingTools/Private/SkeletalMeshModelingToolsEditorMode.cpp
+++ b/Engine/Plugins/Animation/SkeletalMeshModelingTools/Source/SkeletalMeshModelingTools/Private/SkeletalMeshModelingToolsEditorMode.cpp
@@ -718,6 +718,11 @@ void USkeletalMeshModelingToolsEditorMode::HandleMirrorMorphTarget(FName Name)
 	GetCurrentEditingCache()->MirrorMorphTargets({Name});
 }
 
+void USkeletalMeshModelingToolsEditorMode::HandleFlipMorphTarget(FName Name)
+{
+	GetCurrentEditingCache()->FlipMorphTargets({Name});
+}
+
 void USkeletalMeshModelingToolsEditorMode::HandleSetMorphTargetWeight(FName MorphTarget, float Weight)
 {
 	return GetCurrentEditingCache()->HandleSetMorphTargetWeight(MorphTarget, Weight);
@@ -1068,7 +1073,35 @@ void USkeletalMeshModelingToolsEditorMode::InitializeGeometrySelectionSystems()
 			return WeakThis->GetEditingMorphTarget() != NAME_None && WeakThis->IsGeometrySelectionActive();
 		}));
 	CommandList->MapAction(FSkeletalMeshModelingToolsCommands::Get().MirrorEditingMorphTarget, MirrorEditingMorphTargetAction);
-	
+
+	// Register Flip Morph Target operation
+	FUIAction FlipEditingMorphTargetAction = FUIAction(
+		FExecuteAction::CreateLambda([WeakThis]
+		{
+			if (!WeakThis.IsValid() || !WeakThis->GetToolManager() || !WeakThis->GetToolManager()->GetContextTransactionsAPI())
+			{
+				return;
+			}
+
+			FName MorphTargetName = WeakThis->GetEditingMorphTarget();
+			FText Title = FText::Format( LOCTEXT("FlipEditingMorphTransaction", "Flip {0}"), FText::FromName(MorphTargetName));
+			IToolsContextTransactionsAPI* TransactionsAPI = WeakThis->GetToolManager()->GetContextTransactionsAPI();
+			TransactionsAPI->BeginUndoTransaction(Title);
+
+			WeakThis->HandleFlipMorphTarget(MorphTargetName);
+
+			TransactionsAPI->EndUndoTransaction();
+		}),
+		FCanExecuteAction::CreateLambda([WeakThis]()
+		{
+			if (!WeakThis.IsValid())
+			{
+				return false;
+			}
+			return WeakThis->GetEditingMorphTarget() != NAME_None && WeakThis->IsGeometrySelectionActive();
+		}));
+	CommandList->MapAction(FSkeletalMeshModelingToolsCommands::Get().FlipEditingMorphTarget, FlipEditingMorphTargetAction);
+
 	// set up selection type toggles
 	auto RegisterSelectionMode = [WeakThis](UGeometrySelectionManager::EMeshTopologyMode TopoMode, UE::Geometry::EGeometryElementType ElementMode, TSharedPtr<FUICommandInfo> UICommand)
 	{
diff --git a/Engine/Plugins/Animation/SkeletalMeshModelingTools/Source/SkeletalMeshModelingTools/Private/SkeletalMeshModelingToolsEditorMode.h b/Engine/Plugins/Animation/SkeletalMeshModelingTools/Source/SkeletalMeshModelingTools/Private/SkeletalMeshModelingToolsEditorMode.h
--- a/Engine/Plugins/Animation/SkeletalMeshModelingTools/Source/SkeletalMeshModelingTools/Private/SkeletalMeshModelingToolsEditorMode.h
+++ b/Engine/Plugins/Animation/SkeletalMeshModelingTools/Source/SkeletalMeshModelingTools/Private/SkeletalMeshModelingToolsEditorMode.h
@@ -169,8 +169,9 @@ public:
 	FName HandleRenameMorphTarget(FName OldName, FName NameName);
 	void HandleRemoveMorphTargets(const TArray<FName>& Names);
 	TArray<FName> HandleDuplicateMorphTargets(const TArray<FName>& Names);
-	void HandleMirrorMorphTarget(FName Name);	
-	
+	void HandleMirrorMorphTarget(FName Name);
+	void HandleFlipMorphTarget(FName Name);
+
 	void HandleSetMorphTargetWeight(FName MorphTarget, float Weight);
 	bool GetMorphTargetAutoFill(FName MorphTarget);
 	void HandleSetMorphTargetAutoFill(FName MorphTarget, bool bAutoFill, float PreviousOverrideWeight);
diff --git a/Engine/Plugins/Animation/SkeletalMeshModelingTools/Source/SkeletalMeshModelingTools/Private/SkeletalMeshModelingToolsEditorModeToolkit.cpp b/Engine/Plugins/Animation/SkeletalMeshModelingTools/Source/SkeletalMeshModelingTools/Private/SkeletalMeshModelingToolsEditorModeToolkit.cpp
--- a/Engine/Plugins/Animation/SkeletalMeshModelingTools/Source/SkeletalMeshModelingTools/Private/SkeletalMeshModelingToolsEditorModeToolkit.cpp
+++ b/Engine/Plugins/Animation/SkeletalMeshModelingTools/Source/SkeletalMeshModelingTools/Private/SkeletalMeshModelingToolsEditorModeToolkit.cpp
@@ -635,7 +635,8 @@ void FSkeletalMeshModelingToolsEditorModeToolkit::OnEditorModeInitialized()
 	Delegates.OnRemoveMorphTargets.BindUObject(EditorMode.Get(), &USkeletalMeshModelingToolsEditorMode::HandleRemoveMorphTargets);
 	Delegates.OnDuplicateMorphTargets.BindUObject(EditorMode.Get(), &USkeletalMeshModelingToolsEditorMode::HandleDuplicateMorphTargets);
 	Delegates.OnMirrorSelectedMorphTarget.BindUObject(EditorMode.Get(), &USkeletalMeshModelingToolsEditorMode::HandleMirrorMorphTarget);
-	
+	Delegates.OnFlipSelectedMorphTarget.BindUObject(EditorMode.Get(), &USkeletalMeshModelingToolsEditorMode::HandleFlipMorphTarget);
+
 	SAssignNew(MorphTargetManager, SMorphTargetManager)
 			.Delegates(Delegates);
 	MorphTargetManagerWidget->SetContent(MorphTargetManager.ToSharedRef());
diff --git a/Engine/Plugins/Animation/SkeletalMeshModelingTools/Source/SkeletalMeshModelingTools/Private/SkeletalMeshModelingToolsViewportToolbarExtension.cpp b/Engine/Plugins/Animation/SkeletalMeshModelingTools/Source/SkeletalMeshModelingTools/Private/SkeletalMeshModelingToolsViewportToolbarExtension.cpp
--- a/Engine/Plugins/Animation/SkeletalMeshModelingTools/Source/SkeletalMeshModelingTools/Private/SkeletalMeshModelingToolsViewportToolbarExtension.cpp
+++ b/Engine/Plugins/Animation/SkeletalMeshModelingTools/Source/SkeletalMeshModelingTools/Private/SkeletalMeshModelingToolsViewportToolbarExtension.cpp
@@ -290,6 +290,9 @@ namespace UE::SkeletalMeshModelingTools
 								MorphTargetSection.AddEntry(
 									FToolMenuEntry::InitMenuEntryWithCommandList(
 										FSkeletalMeshModelingToolsCommands::Get().MirrorEditingMorphTarget, CommandList));
+								MorphTargetSection.AddEntry(
+									FToolMenuEntry::InitMenuEntryWithCommandList(
+										FSkeletalMeshModelingToolsCommands::Get().FlipEditingMorphTarget, CommandList));
 							}
 						}
 					})
