SKM editor : make the code more resilient to missing LOD to prevent unwanted crashes

#Synced-CL 51755560

diff --git a/Engine/Plugins/Animation/SkeletalMeshModelingTools/Source/SkeletalMeshModelingTools/Private/SkeletalMeshModelingToolsEditorModeToolkit.cpp b/Engine/Plugins/Animation/SkeletalMeshModelingTools/Source/SkeletalMeshModelingTools/Private/SkeletalMeshModelingToolsEditorModeToolkit.cpp
--- a/Engine/Plugins/Animation/SkeletalMeshModelingTools/Source/SkeletalMeshModelingTools/Private/SkeletalMeshModelingToolsEditorModeToolkit.cpp
+++ b/Engine/Plugins/Animation/SkeletalMeshModelingTools/Source/SkeletalMeshModelingTools/Private/SkeletalMeshModelingToolsEditorModeToolkit.cpp
@@ -567,6 +567,11 @@ void FSkeletalMeshModelingToolsEditorModeToolkit::OnEditorModeInitialized()
 
 	constexpr bool bSkipAutoGenerated= true;
 	AssetAvailableLODs = UE::ToolTarget::GetAvailableLODs(SkeletalMesh, bSkipAutoGenerated);
+	ensure(!AssetAvailableLODs.IsEmpty());
+	if (AssetAvailableLODs.IsEmpty())
+	{
+		AssetAvailableLODs.Add(EMeshLODIdentifier::LOD0);
+	}
 	
 	for (EMeshLODIdentifier LOD : AssetAvailableLODs)
 	{
