Fix a crash where transient meshes with no mesh description are opened.

#Synced-CL 52308368

diff --git a/Engine/Plugins/Animation/SkeletalMeshModelingTools/Source/SkeletalMeshModelingTools/Private/SkeletalMeshModelingToolsModule.cpp b/Engine/Plugins/Animation/SkeletalMeshModelingTools/Source/SkeletalMeshModelingTools/Private/SkeletalMeshModelingToolsModule.cpp
--- a/Engine/Plugins/Animation/SkeletalMeshModelingTools/Source/SkeletalMeshModelingTools/Private/SkeletalMeshModelingToolsModule.cpp
+++ b/Engine/Plugins/Animation/SkeletalMeshModelingTools/Source/SkeletalMeshModelingTools/Private/SkeletalMeshModelingToolsModule.cpp
@@ -9,6 +9,7 @@
 
 #include "ContentBrowserMenuContexts.h"
 #include "EditorModeManager.h"
+#include "IPersonaToolkit.h"
 #include "Engine/SkeletalMesh.h"
 #include "ISkeletalMeshEditor.h"
 #include "ISkeletalMeshEditorModule.h"
@@ -144,7 +145,7 @@ TSharedRef<FExtender> FSkeletalMeshModelingToolsModule::ExtendSkelMeshEditorTool
 
 	InCommandList->MapAction(FSkeletalMeshModelingToolsCommands::Get().ToggleEditingToolsMode,
 	    FExecuteAction::CreateRaw(this, &FSkeletalMeshModelingToolsModule::OnToggleEditingToolsMode, EditorPtr),
-	    FCanExecuteAction(),
+	    FCanExecuteAction::CreateRaw(this, &FSkeletalMeshModelingToolsModule::IsMeshEditable, EditorPtr),
 	    FIsActionChecked::CreateRaw(this, &FSkeletalMeshModelingToolsModule::IsEditingToolModeActive, EditorPtr));
 
 	return ToolbarExtender.ToSharedRef();
@@ -196,12 +197,29 @@ void FSkeletalMeshModelingToolsModule::CheckEnableEditingToolModeOnOpen(TWeakPtr
 {
 	bool bEditingModeActive = true;
 	GConfig->GetBool(ConfigSection, ConfigEditingModeActiveKey, bEditingModeActive, GEditorPerProjectIni);
-	if (bEditingModeActive && !IsEditingToolModeActive(InSkeletalMeshEditor))
+	
+	if (bEditingModeActive && !IsEditingToolModeActive(InSkeletalMeshEditor) && IsMeshEditable(InSkeletalMeshEditor))
 	{
 		OnToggleEditingToolsMode(InSkeletalMeshEditor);
 	}
 }
 
+bool FSkeletalMeshModelingToolsModule::IsMeshEditable(TWeakPtr<ISkeletalMeshEditor> InSkeletonEditor)
+{
+	// Fetch mesh to check if it has a mesh description, we can crash if we do not.
+	TSharedPtr<ISkeletalMeshEditor> SkeletalMeshEditor = InSkeletonEditor.Pin();
+	if (!SkeletalMeshEditor.IsValid())
+	{
+		return false;
+	}
+	
+	if (USkeletalMesh* Mesh = SkeletalMeshEditor->GetPersonaToolkit()->GetMesh())
+	{
+		return Mesh->HasMeshDescription(0);
+	}
+	return false;
+}
+
 void FSkeletalMeshModelingToolsModule::OnPostEngineInit()
 {
 	FSkeletalMeshModelingToolsActionCommands::RegisterAllToolActions();
diff --git a/Engine/Plugins/Animation/SkeletalMeshModelingTools/Source/SkeletalMeshModelingTools/Private/SkeletalMeshModelingToolsModule.h b/Engine/Plugins/Animation/SkeletalMeshModelingTools/Source/SkeletalMeshModelingTools/Private/SkeletalMeshModelingToolsModule.h
--- a/Engine/Plugins/Animation/SkeletalMeshModelingTools/Source/SkeletalMeshModelingTools/Private/SkeletalMeshModelingToolsModule.h
+++ b/Engine/Plugins/Animation/SkeletalMeshModelingTools/Source/SkeletalMeshModelingTools/Private/SkeletalMeshModelingToolsModule.h
@@ -32,6 +32,7 @@ private:
 	bool IsEditingToolModeActive(TWeakPtr<ISkeletalMeshEditor> InSkeletalMeshEditor) const;
 	void OnToggleEditingToolsMode(TWeakPtr<ISkeletalMeshEditor> InSkeletalMeshEditor);
 	void CheckEnableEditingToolModeOnOpen(TWeakPtr<ISkeletalMeshEditor> InSkeletalMeshEditor);
+	bool IsMeshEditable(TWeakPtr<ISkeletalMeshEditor> InSkeletonEditor);
 
 	void OnPostEngineInit();
 	
