[SkeletalMeshModelingTools] users can now change between tools quickly by default without having to wait for skeletal mesh build

#Synced-CL 52437958

#Synced-CL 52437958

#Synced-CL 52437958

diff --git a/Engine/Plugins/Animation/SkeletalMeshModelingTools/Source/SkeletalMeshModelingTools/Private/Interfaces/ISKMBackedDynaMeshComponentProvider.h b/Engine/Plugins/Animation/SkeletalMeshModelingTools/Source/SkeletalMeshModelingTools/Private/Interfaces/ISKMBackedDynaMeshComponentProvider.h
--- a/Engine/Plugins/Animation/SkeletalMeshModelingTools/Source/SkeletalMeshModelingTools/Private/Interfaces/ISKMBackedDynaMeshComponentProvider.h
+++ b/Engine/Plugins/Animation/SkeletalMeshModelingTools/Source/SkeletalMeshModelingTools/Private/Interfaces/ISKMBackedDynaMeshComponentProvider.h
@@ -22,5 +22,4 @@ class ISkeletalMeshBackedDynamicMeshComponentProvider
 
 public:
 	virtual USkeletalMeshBackedDynamicMeshComponent* GetComponent(UObject* SourceObject) = 0;
-	virtual bool ShouldCommitToSkeletalMeshOnToolCommit() = 0;
 };
diff --git a/Engine/Plugins/Animation/SkeletalMeshModelingTools/Source/SkeletalMeshModelingTools/Private/SkeletalMeshModelingToolsEditorMode.cpp b/Engine/Plugins/Animation/SkeletalMeshModelingTools/Source/SkeletalMeshModelingTools/Private/SkeletalMeshModelingToolsEditorMode.cpp
--- a/Engine/Plugins/Animation/SkeletalMeshModelingTools/Source/SkeletalMeshModelingTools/Private/SkeletalMeshModelingToolsEditorMode.cpp
+++ b/Engine/Plugins/Animation/SkeletalMeshModelingTools/Source/SkeletalMeshModelingTools/Private/SkeletalMeshModelingToolsEditorMode.cpp
@@ -27,6 +27,7 @@
 #include "Widgets/Notifications/SNotificationList.h"
 #include "Misc/MessageDialog.h"
 #include "HAL/ConsoleManager.h"
+#include "ScopedTransaction.h"
 
 #include UE_INLINE_GENERATED_CPP_BY_NAME(SkeletalMeshModelingToolsEditorMode)
 
@@ -397,8 +398,8 @@ void USkeletalMeshModelingToolsEditorMode::Enter()
 
 	// record switching behavior to restore on exit
 	ToolSwitchModeToRestoreOnExit = GetInteractiveToolsContext()->ToolManager->GetToolSwitchMode();
-	// default to NOT applying changes to skeletal meshes when switching between tools without accepting
-	GetInteractiveToolsContext()->ToolManager->SetToolSwitchMode(EToolManagerToolSwitchMode::CancelIfAble);
+	// tool switch always saves current change to the editing cache
+	GetInteractiveToolsContext()->ToolManager->SetToolSwitchMode(EToolManagerToolSwitchMode::AcceptIfAble);
 
 	GetInteractiveToolsContext()->ToolManager->OnToolPostBuild.AddUObject(this, &USkeletalMeshModelingToolsEditorMode::OnToolPostBuild);
 	GetInteractiveToolsContext()->ToolManager->OnToolEndedWithStatus.AddUObject(this, &USkeletalMeshModelingToolsEditorMode::OnToolEndedWithStatus);
@@ -783,7 +784,7 @@ void USkeletalMeshModelingToolsEditorMode::Exit()
 		{
 			EAppReturnType::Type YesNoReply = FMessageDialog::Open(EAppMsgType::YesNo,
 					FText::Format(
-						LOCTEXT("Prompt_SkeletalMeshModelingToolsClose", "Skeletal Mesh \"{Mesh}\" has {ChangeCount} unapplied changes. Would you like to apply them? (Selecting 'No' will cause all changes to be lost!)"),
+						LOCTEXT("Prompt_SkeletalMeshModelingToolsClose", "Skeletal Mesh \"{0}\" has {1} unapplied changes. Would you like to apply them? (Selecting 'No' will cause all changes to be lost!)"),
 						FText::FromString(GetSkeletalMesh()->GetPathName()),
 						ChangeCount));
 
@@ -1384,6 +1385,13 @@ void USkeletalMeshModelingToolsEditorMode::OnToolEndedWithStatus(UInteractiveToo
 	ToggleBoneManipulation(true);
 
 	bAllowGeometrySelection = true;
+
+	if (bApplyChangesToAssetOnToolEnd)
+	{
+		bApplyChangesToAssetOnToolEnd = false;
+		FScopedTransaction Transaction(LOCTEXT("ApplyChangesToAssetOnToolEnd","Apply Changes To Skeletal Mesh"));
+		ApplyChanges();
+	}
 }
 
 void USkeletalMeshModelingToolsEditorMode::OnToolContextRender(IToolsContextRenderAPI* ToolsContextRenderAPI)
@@ -1633,11 +1641,12 @@ USkeletalMeshBackedDynamicMeshComponent* USkeletalMeshModelingToolsEditorMode::G
 	return GetCurrentEditingCache()->GetEditingMeshComponent();
 }
 
-bool USkeletalMeshModelingToolsEditorMode::ShouldCommitToSkeletalMeshOnToolCommit()
+void USkeletalMeshModelingToolsEditorMode::RequestApplyChangesToAssetOnToolEnd()
 {
-	return ApplyMode == EApplyMode::ApplyOnToolExit;
+	bApplyChangesToAssetOnToolEnd = true;
 }
 
+
 bool USkeletalMeshModelingToolsEditorMode::BeginTransform(const FGizmoState& InState)
 {
 	if (HasActiveGeometrySelection())
@@ -1887,22 +1896,14 @@ void USkeletalMeshModelingToolsEditorMode::DiscardChanges()
 	GetCurrentEditingCache()->DiscardChanges();
 }
 
-USkeletalMeshModelingToolsEditorMode::EApplyMode USkeletalMeshModelingToolsEditorMode::GetApplyMode() const
+USkeletalMeshModelingToolsEditorMode::EToolAcceptAction USkeletalMeshModelingToolsEditorMode::GetToolAcceptAction() const
 {
-	return ApplyMode;
+	return ToolAcceptAction;
 }
 
-void USkeletalMeshModelingToolsEditorMode::SetApplyMode(EApplyMode InApplyMode)
+void USkeletalMeshModelingToolsEditorMode::SetToolAcceptAction(EToolAcceptAction InAction)
 {
-	ApplyMode = InApplyMode;
-	if (ApplyMode == EApplyMode::ApplyManually)
-	{
-		GetInteractiveToolsContext()->ToolManager->SetToolSwitchMode(EToolManagerToolSwitchMode::AcceptIfAble);
-	}
-	else
-	{
-		GetInteractiveToolsContext()->ToolManager->SetToolSwitchMode(EToolManagerToolSwitchMode::CancelIfAble);
-	}
+	ToolAcceptAction = InAction;
 }
 
 void USkeletalMeshModelingToolsEditorMode::HideSkeletonForTool()
diff --git a/Engine/Plugins/Animation/SkeletalMeshModelingTools/Source/SkeletalMeshModelingTools/Private/SkeletalMeshModelingToolsEditorMode.h b/Engine/Plugins/Animation/SkeletalMeshModelingTools/Source/SkeletalMeshModelingTools/Private/SkeletalMeshModelingToolsEditorMode.h
--- a/Engine/Plugins/Animation/SkeletalMeshModelingTools/Source/SkeletalMeshModelingTools/Private/SkeletalMeshModelingToolsEditorMode.h
+++ b/Engine/Plugins/Animation/SkeletalMeshModelingTools/Source/SkeletalMeshModelingTools/Private/SkeletalMeshModelingToolsEditorMode.h
@@ -116,7 +116,7 @@ public:
 
 	// ISkeletalMeshBackedDynamicMeshComponentProvider
 	virtual USkeletalMeshBackedDynamicMeshComponent* GetComponent(UObject* SourceObject) override;
-	bool ShouldCommitToSkeletalMeshOnToolCommit() override;
+
 
 	// IGizmoEdModeInterface
 	virtual bool BeginTransform(const FGizmoState& InState) override;
@@ -142,15 +142,16 @@ public:
 	void ApplyChanges();
 	void DiscardChanges();
 	
-	enum class EApplyMode: uint8
+	enum class EToolAcceptAction: uint8
 	{
-		ApplyOnToolExit,
-		ApplyManually,
+		ApplyToAsset,
+		ExitTool,
 	};
 	
-	EApplyMode GetApplyMode() const;
-	void SetApplyMode(EApplyMode InApplyMode);
+	EToolAcceptAction GetToolAcceptAction() const;
+	void SetToolAcceptAction(EToolAcceptAction InAction);
 
+	void RequestApplyChangesToAssetOnToolEnd();
 	
 	void HideSkeletonForTool();	
 	void ShowSkeletonForTool();
@@ -193,7 +194,8 @@ protected:
 	void OnToolContextRender(IToolsContextRenderAPI* ToolsContextRenderAPI);
 	void OnToolContextDrawHUD(FCanvas* Canvas, IToolsContextRenderAPI* ToolsContextRenderAPI);
 private:
-	EApplyMode ApplyMode = EApplyMode::ApplyOnToolExit;
+	EToolAcceptAction ToolAcceptAction = EToolAcceptAction::ApplyToAsset;
+	bool bApplyChangesToAssetOnToolEnd = false;
 	
 	TSharedPtr<FSkeletalMeshModelingToolsEditorModeBinding> ModeBinding;
 
diff --git a/Engine/Plugins/Animation/SkeletalMeshModelingTools/Source/SkeletalMeshModelingTools/Private/SkeletalMeshModelingToolsEditorModeToolkit.cpp b/Engine/Plugins/Animation/SkeletalMeshModelingTools/Source/SkeletalMeshModelingTools/Private/SkeletalMeshModelingToolsEditorModeToolkit.cpp
--- a/Engine/Plugins/Animation/SkeletalMeshModelingTools/Source/SkeletalMeshModelingTools/Private/SkeletalMeshModelingToolsEditorModeToolkit.cpp
+++ b/Engine/Plugins/Animation/SkeletalMeshModelingTools/Source/SkeletalMeshModelingTools/Private/SkeletalMeshModelingToolsEditorModeToolkit.cpp
@@ -17,6 +17,8 @@
 #include "PropertyEditorModule.h"
 #include "Tools/UEdMode.h"
 #include "Widgets/Input/SButton.h"
+#include "Widgets/Input/SComboButton.h"
+#include "SkeletalMeshModelingToolsStyle.h"
 #include "Widgets/Layout/SBorder.h"
 #include "Widgets/Images/SImage.h"
 #include "Widgets/Text/STextBlock.h"
@@ -438,35 +440,132 @@ void FSkeletalMeshModelingToolsEditorModeToolkit::MakeToolAcceptCancelWidget()
 			.AutoWidth()
 			.Padding(FMargin(0.0, 0.f, 2.f, 0.f))
 			[
-				SNew(SPrimaryButton)
-				.Text_Lambda([this]()
+				SNew(SHorizontalBox)
+				.Visibility_Lambda([this]()
+				{
+					const UEditorInteractiveToolsContext* Context = GetScriptableEditorMode()->GetInteractiveToolsContext();
+					return (Context->ActiveToolHasAccept() || Context->CanCompleteActiveTool()) ? EVisibility::Visible : EVisibility::Collapsed;
+				})
+
+				// Left side: "Apply to Asset" button (green, visible in ApplyOnToolExit mode)
+				+SHorizontalBox::Slot()
+				.AutoWidth()
+				.VAlign(VAlign_Fill)
+				[
+					SNew(SButton)
+					.ButtonStyle(FSkeletalMeshModelingToolsStyle::Get(), "SkeletalMeshModelingTools.AcceptButton.ApplyToAsset")
+					.Text(LOCTEXT("OverlayApplyToAsset", "Apply to Asset"))
+					.ToolTipText(LOCTEXT("OverlayApplyToAssetToolTip", "Accept the current tool change and apply all pending changes to the skeletal mesh asset"))
+					.HAlign(HAlign_Center)
+					.OnClicked_Lambda([this]()
+					{
+						EditorMode->RequestApplyChangesToAssetOnToolEnd();
+						const UEditorInteractiveToolsContext* Context = GetScriptableEditorMode()->GetInteractiveToolsContext();
+						Context->CanAcceptActiveTool()
+							? GetScriptableEditorMode()->GetInteractiveToolsContext()->EndTool(EToolShutdownType::Accept)
+							: GetScriptableEditorMode()->GetInteractiveToolsContext()->EndTool(EToolShutdownType::Completed);
+						return FReply::Handled();
+					})
+					.IsEnabled_Lambda([this]()
+					{
+						const UEditorInteractiveToolsContext* Context = GetScriptableEditorMode()->GetInteractiveToolsContext();
+						return Context->CanAcceptActiveTool() || Context->CanCompleteActiveTool();
+					})
+					.Visibility_Lambda([this]() { return EditorMode->GetToolAcceptAction() == USkeletalMeshModelingToolsEditorMode::EToolAcceptAction::ApplyToAsset ? EVisibility::Visible : EVisibility::Collapsed; })
+				]
+
+				// Left side: "Exit Tool" button (blue, visible in ApplyManually mode)
+				+SHorizontalBox::Slot()
+				.AutoWidth()
+				.VAlign(VAlign_Fill)
+				[
+					SNew(SButton)
+					.ButtonStyle(FSkeletalMeshModelingToolsStyle::Get(), "SkeletalMeshModelingTools.AcceptButton.ExitTool")
+					.Text(LOCTEXT("OverlayExitTool", "Exit Tool"))
+					.ToolTipText(LOCTEXT("OverlayExitToolToolTip", "Accept changes and exit the tool without applying to the asset"))
+					.HAlign(HAlign_Center)
+					.OnClicked_Lambda([this]()
 					{
-						if (EditorMode->GetApplyMode() == USkeletalMeshModelingToolsEditorMode::EApplyMode::ApplyOnToolExit)
+						const UEditorInteractiveToolsContext* Context = GetScriptableEditorMode()->GetInteractiveToolsContext();
+						Context->CanAcceptActiveTool()
+							? GetScriptableEditorMode()->GetInteractiveToolsContext()->EndTool(EToolShutdownType::Accept)
+							: GetScriptableEditorMode()->GetInteractiveToolsContext()->EndTool(EToolShutdownType::Completed);
+						return FReply::Handled();
+					})
+					.IsEnabled_Lambda([this]()
+					{
+						const UEditorInteractiveToolsContext* Context = GetScriptableEditorMode()->GetInteractiveToolsContext();
+						return Context->CanAcceptActiveTool() || Context->CanCompleteActiveTool();
+					})
+					.Visibility_Lambda([this]() { return EditorMode->GetToolAcceptAction() == USkeletalMeshModelingToolsEditorMode::EToolAcceptAction::ExitTool ? EVisibility::Visible : EVisibility::Collapsed; })
+				]
+
+				// Right side: dropdown arrow for alternative action
+				+SHorizontalBox::Slot()
+				.AutoWidth()
+				.VAlign(VAlign_Fill)
+				[
+					SNew(SComboButton)
+					.HasDownArrow(true)
+					.ContentPadding(0.f)
+					.ComboButtonStyle(FSkeletalMeshModelingToolsStyle::Get(), "SkeletalMeshModelingTools.AcceptButton.ComboButton")
+					.MenuPlacement(MenuPlacement_ComboBoxRight)
+					.OnGetMenuContent_Lambda([this]() -> TSharedRef<SWidget>
+					{
+						FMenuBuilder MenuBuilder(/*bInShouldCloseWindowAfterMenuSelection=*/true, nullptr);
+
+						const UEditorInteractiveToolsContext* Context = GetScriptableEditorMode()->GetInteractiveToolsContext();
+						const bool bCanAccept = Context->CanAcceptActiveTool();
+						const FText AcceptOrComplete = bCanAccept
+							? LOCTEXT("AcceptVerb", "Accept")
+							: LOCTEXT("CompleteVerb", "Complete");
+
+						if (EditorMode->GetToolAcceptAction() == USkeletalMeshModelingToolsEditorMode::EToolAcceptAction::ApplyToAsset)
 						{
-							return LOCTEXT("OverlayApplyToAsset", "Apply to Asset");
+							MenuBuilder.AddMenuEntry(
+								LOCTEXT("SwitchToExitTool", "Apply Later"),
+								LOCTEXT("SwitchToExitToolTip", "Exit the tool without applying changes to the asset. Changes can be applied later."),
+								FSlateIcon(),
+								FUIAction(FExecuteAction::CreateLambda([this]()
+								{
+									EditorMode->SetToolAcceptAction(USkeletalMeshModelingToolsEditorMode::EToolAcceptAction::ExitTool);
+									const UEditorInteractiveToolsContext* Ctx = GetScriptableEditorMode()->GetInteractiveToolsContext();
+									Ctx->CanAcceptActiveTool()
+										? GetScriptableEditorMode()->GetInteractiveToolsContext()->EndTool(EToolShutdownType::Accept)
+										: GetScriptableEditorMode()->GetInteractiveToolsContext()->EndTool(EToolShutdownType::Completed);
+								}))
+							);
 						}
-						
-						check(EditorMode->GetApplyMode() == USkeletalMeshModelingToolsEditorMode::EApplyMode::ApplyManually)
-						
-						return LOCTEXT("OverlayExitTool", "Exit Tool");
+						else
+						{
+							const bool bIsAcceptTool = GetScriptableEditorMode()->GetInteractiveToolsContext()->CanAcceptActiveTool();
+							int32 ChangeCount = EditorMode->GetCurrentEditingCache()->GetEditingMeshComponent()->GetChangeCount();
+							ChangeCount += (bIsAcceptTool ? 1 : 0);
+								
+							MenuBuilder.AddMenuEntry(
+								FText::Format(LOCTEXT("SwitchToApplyToAsset", "Apply {0} {0}|plural(one=Change,other=Changes) to Asset"), FText::AsNumber(ChangeCount)),
+								LOCTEXT("SwitchToApplyToAssetTip", "Apply the current tool change and all pending changes to the skeletal mesh asset."),
+								FSlateIcon(),
+								FUIAction(FExecuteAction::CreateLambda([this]()
+								{
+									EditorMode->SetToolAcceptAction(USkeletalMeshModelingToolsEditorMode::EToolAcceptAction::ApplyToAsset);
+									EditorMode->RequestApplyChangesToAssetOnToolEnd();
+									const UEditorInteractiveToolsContext* Ctx = GetScriptableEditorMode()->GetInteractiveToolsContext();
+									Ctx->CanAcceptActiveTool()
+										? GetScriptableEditorMode()->GetInteractiveToolsContext()->EndTool(EToolShutdownType::Accept)
+										: GetScriptableEditorMode()->GetInteractiveToolsContext()->EndTool(EToolShutdownType::Completed);
+								}))
+							);
+						}
+
+						return MenuBuilder.MakeWidget();
 					})
-				.ToolTipText_Lambda([this]()
-				{
-					if (EditorMode->GetApplyMode() == USkeletalMeshModelingToolsEditorMode::EApplyMode::ApplyOnToolExit)
+					.IsEnabled_Lambda([this]()
 					{
-						return LOCTEXT("OverlayApplyToAssetToolTip",
-							"\"Apply on Tool Exit\" mode is active: changes are applied to the asset immediately. "
-							"Switching to a different tool discards current change");
-					}
-						
-					check(EditorMode->GetApplyMode() == USkeletalMeshModelingToolsEditorMode::EApplyMode::ApplyManually)
-						
-					return LOCTEXT("OverlayExitToolToolTip",
-						"\"Apply Manually\" mode is active: changes are saved when Exiting tool or Switching to a different tool");
-				})
-				.OnClicked_Lambda([this]() { GetScriptableEditorMode()->GetInteractiveToolsContext()->EndTool(EToolShutdownType::Accept); return FReply::Handled(); })
-				.IsEnabled_Lambda([this]() { return GetScriptableEditorMode()->GetInteractiveToolsContext()->CanAcceptActiveTool(); })
-				.Visibility_Lambda([this]() { return GetScriptableEditorMode()->GetInteractiveToolsContext()->ActiveToolHasAccept() ? EVisibility::Visible : EVisibility::Collapsed; })
+						const UEditorInteractiveToolsContext* Context = GetScriptableEditorMode()->GetInteractiveToolsContext();
+						return Context->CanAcceptActiveTool() || Context->CanCompleteActiveTool();
+					})
+				]
 			]
 
 			+SHorizontalBox::Slot()
@@ -474,7 +573,7 @@ void FSkeletalMeshModelingToolsEditorModeToolkit::MakeToolAcceptCancelWidget()
 			.Padding(FMargin(2.0, 0.f, 0.f, 0.f))
 			[
 				SNew(SButton)
-				.TextStyle( FAppStyle::Get(), "DialogButtonText" )
+				.ButtonStyle(FSkeletalMeshModelingToolsStyle::Get(), "SkeletalMeshModelingTools.CancelButton")
 				.Text(LOCTEXT("OverlayCancel", "Cancel"))
 				.ToolTipText(LOCTEXT("OverlayCancelTooltip", "Cancel the active Tool [Esc]"))
 				.HAlign(HAlign_Center)
@@ -482,18 +581,6 @@ void FSkeletalMeshModelingToolsEditorModeToolkit::MakeToolAcceptCancelWidget()
 				.IsEnabled_Lambda([this]() { return GetScriptableEditorMode()->GetInteractiveToolsContext()->CanCancelActiveTool(); })
 				.Visibility_Lambda([this]() { return GetScriptableEditorMode()->GetInteractiveToolsContext()->ActiveToolHasAccept() ? EVisibility::Visible : EVisibility::Collapsed; })
 			]
-
-			+SHorizontalBox::Slot()
-			.AutoWidth()
-			.Padding(FMargin(2.0, 0.f, 0.f, 0.f))
-			[
-				SNew(SPrimaryButton)
-				.Text(LOCTEXT("OverlayComplete", "Complete"))
-				.ToolTipText(LOCTEXT("OverlayCompleteTooltip", "Exit the active Tool [Enter]"))
-				.OnClicked_Lambda([this]() { GetScriptableEditorMode()->GetInteractiveToolsContext()->EndTool(EToolShutdownType::Completed); return FReply::Handled(); })
-				.IsEnabled_Lambda([this]() { return GetScriptableEditorMode()->GetInteractiveToolsContext()->CanCompleteActiveTool(); })
-				.Visibility_Lambda([this]() { return GetScriptableEditorMode()->GetInteractiveToolsContext()->CanCompleteActiveTool() ? EVisibility::Visible : EVisibility::Collapsed; })
-			]
 		]	
 	];
 }
@@ -582,9 +669,6 @@ void FSkeletalMeshModelingToolsEditorModeToolkit::OnEditorModeInitialized()
 
 	AssetLODMode->SetSelectedItem(AssetLODModes[0]);
 
-	int32 EnumIndex = ApplyModeEnums.IndexOfByKey(EditorMode->GetApplyMode());
-	ApplyModeComboBox->SetSelectedItem(ApplyModeStrings[EnumIndex]);
-
 	// Create a Skeleton Tree for dynamic mesh skeleton
 	SAssignNew(DynamicMeshSkeletonWidget, SBorder)
 	.Padding(4.f)
@@ -806,37 +890,7 @@ TSharedPtr<SWidget> FSkeletalMeshModelingToolsEditorModeToolkit::MakeFooterWidge
 		.ToolTipText(LOCTEXT("ActiveLODLabelToolTip", "Select the LOD to be used when editing an existing mesh."));
 	}
 
-	// Apply Mode Picker Widget
-	{
-		FText ToolTip = LOCTEXT("ApplyModeToolTip", "Select when tool changes should be applied to the asset.\n"
-								"When using \"Apply Manually\", you can switch between different tools directly without waiting for changes to be applied to the asset");
-		ApplyModeLabel =
-			SNew(STextBlock)
-			.Text(LOCTEXT("ApplyModeLabel", "Apply Mode"))
-			.ToolTipText(ToolTip);
-
-		ApplyModeEnums.Add(USkeletalMeshModelingToolsEditorMode::EApplyMode::ApplyOnToolExit);
-		ApplyModeStrings.Add(MakeShared<FString>(TEXT("Apply On Tool Exit")));
-		
-		ApplyModeEnums.Add(USkeletalMeshModelingToolsEditorMode::EApplyMode::ApplyManually);
-		ApplyModeStrings.Add(MakeShared<FString>(TEXT("Apply Manually")));
-			
-		ApplyModeComboBox =
-			SNew(STextComboBox)
-			.ToolTipText(ToolTip)
-			.OptionsSource(&ApplyModeStrings)
-			.IsEnabled(this, &FSkeletalMeshModelingToolsEditorModeToolkit::CanChangeAssetEditingSettings)
-			.OnSelectionChanged_Lambda([this](TSharedPtr<FString> String, ESelectInfo::Type)
-			{
-				int32 Index = ApplyModeStrings.IndexOfByKey(String);
-				USkeletalMeshModelingToolsEditorMode::EApplyMode ApplyMode = ApplyModeEnums[Index];
 
-				if (EditorMode.IsValid())
-				{
-					EditorMode->SetApplyMode(ApplyMode);
-				}
-			});
-	}
 	
 
 	
@@ -855,16 +909,6 @@ TSharedPtr<SWidget> FSkeletalMeshModelingToolsEditorModeToolkit::MakeFooterWidge
 				+ SHorizontalBox::Slot().Padding(0).HAlign(HAlign_Fill).Padding(0).FillWidth(4.f)
 					[ AssetLODMode->AsShared() ]
 			]
-			+SVerticalBox::Slot()
-			.AutoHeight()
-			[
-				SNew(SHorizontalBox)
-				+ SHorizontalBox::Slot().Padding(0).HAlign(HAlign_Left).VAlign(VAlign_Center).FillWidth(2.f)
-					[ ApplyModeLabel->AsShared() ]
-				+ SHorizontalBox::Slot().Padding(0).HAlign(HAlign_Fill).Padding(0).FillWidth(4.f)
-					[ ApplyModeComboBox->AsShared() ]
-			]
-			
 			+SVerticalBox::Slot()
 			.AutoHeight()
 			[
diff --git a/Engine/Plugins/Animation/SkeletalMeshModelingTools/Source/SkeletalMeshModelingTools/Private/SkeletalMeshModelingToolsEditorModeToolkit.h b/Engine/Plugins/Animation/SkeletalMeshModelingTools/Source/SkeletalMeshModelingTools/Private/SkeletalMeshModelingToolsEditorModeToolkit.h
--- a/Engine/Plugins/Animation/SkeletalMeshModelingTools/Source/SkeletalMeshModelingTools/Private/SkeletalMeshModelingToolsEditorModeToolkit.h
+++ b/Engine/Plugins/Animation/SkeletalMeshModelingTools/Source/SkeletalMeshModelingTools/Private/SkeletalMeshModelingToolsEditorModeToolkit.h
@@ -102,10 +102,6 @@ private:
 	TSharedPtr<STextBlock> AssetLODModeLabel;
 	TSharedPtr<STextComboBox> AssetLODMode;
 
-	TArray<USkeletalMeshModelingToolsEditorMode::EApplyMode> ApplyModeEnums;
-	TArray<TSharedPtr<FString>> ApplyModeStrings;
-	TSharedPtr<STextBlock> ApplyModeLabel;
-	TSharedPtr<STextComboBox> ApplyModeComboBox;
 	
 	TSharedPtr<SWidget> DefaultSkeletonWidget;
 	TSharedPtr<SBorder> DynamicMeshSkeletonWidget;
diff --git a/Engine/Plugins/Animation/SkeletalMeshModelingTools/Source/SkeletalMeshModelingTools/Private/SkeletalMeshModelingToolsStyle.cpp b/Engine/Plugins/Animation/SkeletalMeshModelingTools/Source/SkeletalMeshModelingTools/Private/SkeletalMeshModelingToolsStyle.cpp
--- a/Engine/Plugins/Animation/SkeletalMeshModelingTools/Source/SkeletalMeshModelingTools/Private/SkeletalMeshModelingToolsStyle.cpp
+++ b/Engine/Plugins/Animation/SkeletalMeshModelingTools/Source/SkeletalMeshModelingTools/Private/SkeletalMeshModelingToolsStyle.cpp
@@ -7,12 +7,13 @@
 #include "Styling/SlateStyleRegistry.h"
 #include "Styling/SlateTypes.h"
 #include "Styling/StyleColors.h"
+#include "Styling/CoreStyle.h"
 
-FSkeletalMeshModelingToolsStyle::FSkeletalMeshModelingToolsStyle() 
+FSkeletalMeshModelingToolsStyle::FSkeletalMeshModelingToolsStyle()
 	: FSlateStyleSet("SkeletalMeshModelingToolsStyle")
 {
 	FCheckBoxStyle BaseCheckBoxStyle = FAppStyle::Get().GetWidgetStyle<FCheckBoxStyle>("ToggleButtonCheckbox");
-	Set("SkeletalMeshModelingTools.ToggleButtonCheckbox", 
+	Set("SkeletalMeshModelingTools.ToggleButtonCheckbox",
 		BaseCheckBoxStyle
 		.SetUncheckedImage(FSlateRoundedBoxBrush(FStyleColors::Header, 4.0f, FStyleColors::Input, 1.0f))
 		.SetUncheckedHoveredImage(FSlateRoundedBoxBrush(FStyleColors::Hover, 4.0f, FStyleColors::Input, 1.0f))
@@ -22,6 +23,73 @@ FSkeletalMeshModelingToolsStyle::FSkeletalMeshModelingToolsStyle()
 		.SetCheckedPressedImage(FSlateRoundedBoxBrush(FStyleColors::PrimaryHover, 4.0f, FStyleColors::Input, 1.0f))
 		.SetPadding(FMargin(16, 6))
 	);
+
+	// Split accept button styles based on default "Button" style.
+	// No outlines — matches viewport toolbar split button approach to avoid double borders at seams.
+	{
+		const float Radius = 4.0f;
+		const FVector4 StartCorners(Radius, 0.f, 0.f, Radius);
+		const FVector4 EndCorners(0.f, Radius, Radius, 0.f);
+		const FVector4 AllCorners(Radius, Radius, Radius, Radius);
+
+		const FButtonStyle& BaseButton = FCoreStyle::Get().GetWidgetStyle<FButtonStyle>("Button");
+
+		// "Apply to Asset" button (green) — matches FlatButton.Success colors
+		// Color multipliers from StarshipStyle ButtonColor: Normal=0.8, Hovered=1.0, Pressed=0.6
+		{
+			const FLinearColor SuccessBase(0.10616f, 0.48777f, 0.10616f);
+			FLinearColor SuccessNormal = SuccessBase * 0.8f; SuccessNormal.A = 1.0f;
+			FLinearColor SuccessHovered = SuccessBase * 1.0f; SuccessHovered.A = 1.0f;
+			FLinearColor SuccessPressed = SuccessBase * 0.6f; SuccessPressed.A = 1.0f;
+
+			Set("SkeletalMeshModelingTools.AcceptButton.ApplyToAsset",
+				FButtonStyle(BaseButton)
+				.SetNormal(FSlateRoundedBoxBrush(SuccessNormal, StartCorners))
+				.SetHovered(FSlateRoundedBoxBrush(SuccessHovered, StartCorners))
+				.SetPressed(FSlateRoundedBoxBrush(SuccessPressed, StartCorners))
+				.SetDisabled(FSlateRoundedBoxBrush(FStyleColors::Dropdown, StartCorners))
+			);
+		}
+
+		// "Exit Tool" button (blue) — matches PrimaryButton colors via FStyleColors
+		{
+			Set("SkeletalMeshModelingTools.AcceptButton.ExitTool",
+				FButtonStyle(BaseButton)
+				.SetNormal(FSlateRoundedBoxBrush(FStyleColors::Primary, StartCorners))
+				.SetHovered(FSlateRoundedBoxBrush(FStyleColors::PrimaryHover, StartCorners))
+				.SetPressed(FSlateRoundedBoxBrush(FStyleColors::PrimaryPress, StartCorners))
+				.SetDisabled(FSlateRoundedBoxBrush(FStyleColors::Dropdown, StartCorners))
+			);
+		}
+
+		// Combo dropdown button (right side) — default button colors with rounded right corners
+		{
+			Set("SkeletalMeshModelingTools.AcceptButton.ComboButton",
+				FComboButtonStyle(FCoreStyle::Get().GetWidgetStyle<FComboButtonStyle>("ComboButton"))
+				.SetButtonStyle(
+					FButtonStyle(BaseButton)
+					.SetNormal(FSlateRoundedBoxBrush(FStyleColors::Secondary, EndCorners))
+					.SetHovered(FSlateRoundedBoxBrush(FStyleColors::Hover, EndCorners))
+					.SetPressed(FSlateRoundedBoxBrush(FStyleColors::Header, EndCorners))
+					.SetDisabled(FSlateRoundedBoxBrush(FStyleColors::Dropdown, EndCorners))
+					.SetNormalPadding(0.f)
+					.SetPressedPadding(0.f)
+				)
+				.SetDownArrowPadding(FMargin(2.0f))
+			);
+		}
+
+		// Cancel button — default button colors, fully rounded, no outline
+		{
+			Set("SkeletalMeshModelingTools.CancelButton",
+				FButtonStyle(BaseButton)
+				.SetNormal(FSlateRoundedBoxBrush(FStyleColors::Secondary, AllCorners))
+				.SetHovered(FSlateRoundedBoxBrush(FStyleColors::Hover, AllCorners))
+				.SetPressed(FSlateRoundedBoxBrush(FStyleColors::Header, AllCorners))
+				.SetDisabled(FSlateRoundedBoxBrush(FStyleColors::Dropdown, AllCorners))
+			);
+		}
+	}
 }
 
 FSkeletalMeshModelingToolsStyle& FSkeletalMeshModelingToolsStyle::Get()
diff --git a/Engine/Plugins/Animation/SkeletalMeshModelingTools/Source/SkeletalMeshModelingTools/Private/ToolTargets/SKMBackedDynaMeshComponentToolTarget.cpp b/Engine/Plugins/Animation/SkeletalMeshModelingTools/Source/SkeletalMeshModelingTools/Private/ToolTargets/SKMBackedDynaMeshComponentToolTarget.cpp
--- a/Engine/Plugins/Animation/SkeletalMeshModelingTools/Source/SkeletalMeshModelingTools/Private/ToolTargets/SKMBackedDynaMeshComponentToolTarget.cpp
+++ b/Engine/Plugins/Animation/SkeletalMeshModelingTools/Source/SkeletalMeshModelingTools/Private/ToolTargets/SKMBackedDynaMeshComponentToolTarget.cpp
@@ -31,11 +31,6 @@ void USkeletalMeshBackedDynamicMeshComponentToolTarget::CommitSkeletonModifier(U
 	Super::CommitSkeletonModifier(InModifier);
 
 	GetComponent()->HandleSkeletonChange(InModifier);
-
-	if (bCommitToSkeletalMeshOnToolCommit)
-	{
-		GetComponent()->CommitToSkeletalMesh();
-	}
 }
 
 void USkeletalMeshBackedDynamicMeshComponentToolTarget::CommitDynamicMeshChange(TUniquePtr<FToolCommandChange> Change, const FText& ChangeMessage)
@@ -50,11 +45,6 @@ void USkeletalMeshBackedDynamicMeshComponentToolTarget::CommitDynamicMeshChange(
 	// so we have to assume that it is dirty such that symmetry tracker is always up to date.
 	GetComponent()->MarkBaseGeometryDirty();
 	GetComponent()->MarkDirty();
-
-	if (bCommitToSkeletalMeshOnToolCommit)
-	{
-		GetComponent()->CommitToSkeletalMesh();
-	}
 }
 
 USkeletalMesh* USkeletalMeshBackedDynamicMeshComponentToolTarget::GetSkeletalMesh() const
@@ -150,8 +140,6 @@ UToolTarget* USkeletalMeshBackedDynamicMeshComponentToolTargetFactory::BuildTarg
 	Target->InitializeComponent(Cast<UDynamicMeshComponent>(Component));
 	checkSlow(Target->Component.IsValid() && Requirements.AreSatisfiedBy(Target));
 
-	Target->bCommitToSkeletalMeshOnToolCommit = ComponentProvider->ShouldCommitToSkeletalMeshOnToolCommit();
-	
 	return Target;
 	
 }
diff --git a/Engine/Plugins/Animation/SkeletalMeshModelingTools/Source/SkeletalMeshModelingTools/Private/ToolTargets/SKMBackedDynaMeshComponentToolTarget.h b/Engine/Plugins/Animation/SkeletalMeshModelingTools/Source/SkeletalMeshModelingTools/Private/ToolTargets/SKMBackedDynaMeshComponentToolTarget.h
--- a/Engine/Plugins/Animation/SkeletalMeshModelingTools/Source/SkeletalMeshModelingTools/Private/ToolTargets/SKMBackedDynaMeshComponentToolTarget.h
+++ b/Engine/Plugins/Animation/SkeletalMeshModelingTools/Source/SkeletalMeshModelingTools/Private/ToolTargets/SKMBackedDynaMeshComponentToolTarget.h
@@ -33,8 +33,6 @@ protected:
 	USkeletalMeshBackedDynamicMeshComponent* GetComponent() const;
 
 	TWeakObjectPtr<USkeletalMesh> WeakSkeletalMesh;
-
-	bool bCommitToSkeletalMeshOnToolCommit = true;
 	
 	friend class USkeletalMeshBackedDynamicMeshComponentToolTargetFactory;
 };
