[SKMModelingTools] SMorphTargetManager UX: double-click weight toggle, slider 0-1 range, rename via right click menu only

#Synced-CL 52437958

diff --git a/Engine/Plugins/Animation/SkeletalMeshModelingTools/Source/SkeletalMeshModelingTools/Private/SMorphTargetManager.cpp b/Engine/Plugins/Animation/SkeletalMeshModelingTools/Source/SkeletalMeshModelingTools/Private/SMorphTargetManager.cpp
--- a/Engine/Plugins/Animation/SkeletalMeshModelingTools/Source/SkeletalMeshModelingTools/Private/SMorphTargetManager.cpp
+++ b/Engine/Plugins/Animation/SkeletalMeshModelingTools/Source/SkeletalMeshModelingTools/Private/SMorphTargetManager.cpp
@@ -25,6 +25,7 @@ namespace MorphTargetManagerLocal
 	{
 		FName Name = NAME_None;
 		float Weight = 1.0f;
+		bool bIsRenaming = false;
 		TSharedPtr<SInlineEditableTextBlock> EditableText;
 	};
 	static const FName ColumnId_MorphTargetNameLabel( "MorphTargetName" );
@@ -91,6 +92,7 @@ public:
 					.HighlightText(Manager,  &SMorphTargetManager::GetHighlightText, MorphNameText)
 					.OnTextCommitted(this, &SMorphTargetManagerListRow::OnRenameMorphTarget)
 					.IsSelected(this, &SMorphTargetManagerListRow::IsSelected)
+					.IsReadOnly_Lambda([this]() { return !Item->bIsRenaming; })
 				]
 				+ SHorizontalBox::Slot()
 				.AutoWidth()
@@ -116,7 +118,7 @@ public:
 				.VAlign( VAlign_Center )
 				[
 					SNew( SSpinBox<float> )
-					.MinSliderValue(-1.f)
+					.MinSliderValue(0.f)
 					.MaxSliderValue(1.f)
 					.Value( this, &SMorphTargetManagerListRow::GetWeight )
 					.OnBeginSliderMovement(this, &SMorphTargetManagerListRow::OnBeginSlideMorphTargetWeight)
@@ -222,6 +224,7 @@ private:
 
 	void OnRenameMorphTarget(const FText& InNewName, ETextCommit::Type)
 	{
+		Item->bIsRenaming = false;
 		Item->Name = Manager->RenameMorphTarget( Item->Name, *InNewName.ToString());
 	}
 
@@ -295,6 +298,7 @@ void SMorphTargetManager::Construct(const FArguments& InArgs)
 			SAssignNew( ListView , SMorphTargetManagerListType )
 			.ListItemsSource( &List)
 			.OnGenerateRow(this, &SMorphTargetManager::GenerateMorphTargetRow )
+			.OnMouseButtonDoubleClick(this, &SMorphTargetManager::OnDoubleClickMorphTarget)
 			.OnContextMenuOpening( this, &SMorphTargetManager::OnGetContextMenuContent )
 			.HeaderRow
 			(
@@ -530,6 +534,18 @@ void SMorphTargetManager::AddMorphTarget()
 	SelectMorphTargets({NewMorphTarget});
 }
 
+void SMorphTargetManager::OnDoubleClickMorphTarget(FMorphTargetInfoPtr Item)
+{
+	if (!Item.IsValid())
+	{
+		return;
+	}
+
+	const float CurrentWeight = GetMorphTargetWeight(Item->Name);
+	const float NewWeight = (CurrentWeight != 1.0f) ? 1.0f : 0.0f;
+	SetMorphTargetWeight(Item->Name, NewWeight);
+}
+
 FName SMorphTargetManager::RenameMorphTarget(FName InOldName, FName InNewName)
 {
 	FScopedTransaction Transaction(LOCTEXT("RenameMorphTarget", "Rename Morph Target"));
@@ -539,6 +555,7 @@ FName SMorphTargetManager::RenameMorphTarget(FName InOldName, FName InNewName)
 void SMorphTargetManager::RenameSelectedMorphTarget()
 {
 	TArray<FMorphTargetInfoPtr> Items = ListView->GetSelectedItems();
+	Items[0]->bIsRenaming = true;
 	Items[0]->EditableText->EnterEditingMode();
 }
 
diff --git a/Engine/Plugins/Animation/SkeletalMeshModelingTools/Source/SkeletalMeshModelingTools/Private/SMorphTargetManager.h b/Engine/Plugins/Animation/SkeletalMeshModelingTools/Source/SkeletalMeshModelingTools/Private/SMorphTargetManager.h
--- a/Engine/Plugins/Animation/SkeletalMeshModelingTools/Source/SkeletalMeshModelingTools/Private/SMorphTargetManager.h
+++ b/Engine/Plugins/Animation/SkeletalMeshModelingTools/Source/SkeletalMeshModelingTools/Private/SMorphTargetManager.h
@@ -95,7 +95,8 @@ public:
 	bool IsEditingMorphTarget(FName MorphTarget);
 
 	void AddMorphTarget();
-	FName RenameMorphTarget(FName InNewName, FName InOldName);
+	FName RenameMorphTarget(FName InOldName, FName InNewName);
+	void OnDoubleClickMorphTarget(FMorphTargetInfoPtr Item);
 protected:
 	void RenameSelectedMorphTarget();
 	bool CanRename();
