SMorphTargetManager: decouple editing morph target selection from weight override

#Synced-CL 52437958

diff --git a/Engine/Plugins/Animation/SkeletalMeshModelingTools/Source/SkeletalMeshModelingTools/Private/SkeletalMeshModelingToolsEditorMode.cpp b/Engine/Plugins/Animation/SkeletalMeshModelingTools/Source/SkeletalMeshModelingTools/Private/SkeletalMeshModelingToolsEditorMode.cpp
--- a/Engine/Plugins/Animation/SkeletalMeshModelingTools/Source/SkeletalMeshModelingTools/Private/SkeletalMeshModelingToolsEditorMode.cpp
+++ b/Engine/Plugins/Animation/SkeletalMeshModelingTools/Source/SkeletalMeshModelingTools/Private/SkeletalMeshModelingToolsEditorMode.cpp
@@ -644,11 +644,7 @@ void USkeletalMeshModelingToolsEditorMode::HandleSetEditingMorphTarget(FName InM
 {
 	// Such that EditingMorphTarget is recorded for undo/redo 
 	Modify();
-	FName OldEditingMorphTarget = EditingMorphTarget;
 	EditingMorphTarget = InMorphTarget;
-
-	GetCurrentEditingCache()->ClearMorphTargetOverride(OldEditingMorphTarget);
-	GetCurrentEditingCache()->OverrideMorphTargetWeight(EditingMorphTarget, 1.0f);
 }
 
 
@@ -677,6 +673,7 @@ FName USkeletalMeshModelingToolsEditorMode::HandleAddMorphTarget(FName InName)
 	FName NewMorphTarget = GetCurrentEditingCache()->AddMorphTarget(InName);
 	
 	HandleSetEditingMorphTarget(NewMorphTarget);
+	GetCurrentEditingCache()->OverrideMorphTargetWeight(NewMorphTarget, 1.0f);
 	
 	return NewMorphTarget;
 }
