[SkeletalMeshModelingTools] Allowing using geometry selection mode buttons while a tool is active. It accepts the current tool and activates the corresponding selection mode

#Synced-CL 52437958

diff --git a/Engine/Plugins/Animation/SkeletalMeshModelingTools/Source/SkeletalMeshModelingTools/Private/SkeletalMeshModelingToolsEditorMode.cpp b/Engine/Plugins/Animation/SkeletalMeshModelingTools/Source/SkeletalMeshModelingTools/Private/SkeletalMeshModelingToolsEditorMode.cpp
--- a/Engine/Plugins/Animation/SkeletalMeshModelingTools/Source/SkeletalMeshModelingTools/Private/SkeletalMeshModelingToolsEditorMode.cpp
+++ b/Engine/Plugins/Animation/SkeletalMeshModelingTools/Source/SkeletalMeshModelingTools/Private/SkeletalMeshModelingToolsEditorMode.cpp
@@ -1129,7 +1129,17 @@ void USkeletalMeshModelingToolsEditorMode::InitializeGeometrySelectionSystems()
 			FExecuteAction::CreateLambda([WeakThis, TopoMode, ElementMode]() {
 				if ( WeakThis.IsValid() && WeakThis->GetToolManager() && WeakThis->GetToolManager()->GetContextTransactionsAPI())
 				{
-					IToolsContextTransactionsAPI* TransactionsAPI = WeakThis->GetToolManager()->GetContextTransactionsAPI();
+					// Accept/complete active tool before switching selection mode
+					UInteractiveToolManager* ToolManager = WeakThis->GetToolManager();
+					if (ToolManager->HasActiveTool(EToolSide::Left))
+					{
+						EToolShutdownType ShutdownType = ToolManager->CanAcceptActiveTool(EToolSide::Left)
+							? EToolShutdownType::Accept
+							: EToolShutdownType::Completed;
+						ToolManager->DeactivateTool(EToolSide::Left, ShutdownType);
+					}
+
+					IToolsContextTransactionsAPI* TransactionsAPI = ToolManager->GetContextTransactionsAPI();
 					TransactionsAPI->BeginUndoTransaction(LOCTEXT("ChangeSelectionMode", "Selection Mode"));
 					WeakThis->SelectionManager->SetMeshSelectionTypeAndMode(ElementMode, TopoMode,  TopoMode != UGeometrySelectionManager::EMeshTopologyMode::None);
 					TransactionsAPI->EndUndoTransaction();
@@ -1140,24 +1150,22 @@ void USkeletalMeshModelingToolsEditorMode::InitializeGeometrySelectionSystems()
 						ModelingEditorSettings->LastMeshSelectionElementType = static_cast<int>(ElementMode);
 						ModelingEditorSettings->SaveConfig();
 					}
-					
-					
-				}
-			}),
-			FCanExecuteAction::CreateLambda([WeakThis]()
-			{
-				if (!WeakThis.IsValid())
-				{
-					return false;
 				}
-				return WeakThis->bAllowGeometrySelection;
 			}),
+			FCanExecuteAction(),
 			FIsActionChecked::CreateLambda([WeakThis, TopoMode, ElementMode]()
 				{
 					if (!WeakThis.IsValid() || !WeakThis->SelectionManager)
 					{
 						return false;
 					}
+					
+					// Don't want to show the current active mode when a tool is running, it would be quite confusing
+					if (WeakThis->bSuspendGeometrySelection)
+					{
+						return false;
+					}
+
 					return WeakThis->SelectionManager->GetMeshTopologyMode() == TopoMode && WeakThis->SelectionManager->GetSelectionElementType() == ElementMode;
 				}),
 			EUIActionRepeatMode::RepeatDisabled);
@@ -1201,7 +1209,7 @@ void USkeletalMeshModelingToolsEditorMode::PostBuildViewportInteractions(const U
 
 bool USkeletalMeshModelingToolsEditorMode::IsGeometrySelectionActive() const
 {
-	return bAllowGeometrySelection && SelectionManager->GetMeshTopologyMode() != UGeometrySelectionManager::EMeshTopologyMode::None;
+	return !bSuspendGeometrySelection && SelectionManager->GetMeshTopologyMode() != UGeometrySelectionManager::EMeshTopologyMode::None;
 }
 
 bool USkeletalMeshModelingToolsEditorMode::HasActiveGeometrySelection() const
@@ -1382,7 +1390,7 @@ void USkeletalMeshModelingToolsEditorMode::OnToolPostBuild( UInteractiveToolMana
 	ToggleBoneManipulation(false);
 
 	// Disable element selection modes while tool is active
-	bAllowGeometrySelection = false;
+	bSuspendGeometrySelection = true;
 }
 
 void USkeletalMeshModelingToolsEditorMode::OnToolEndedWithStatus(UInteractiveToolManager* Manager, UInteractiveTool* Tool,
@@ -1397,7 +1405,7 @@ void USkeletalMeshModelingToolsEditorMode::OnToolEndedWithStatus(UInteractiveToo
 	
 	ToggleBoneManipulation(true);
 
-	bAllowGeometrySelection = true;
+	bSuspendGeometrySelection = false;
 
 	if (bApplyChangesToAssetOnToolEnd)
 	{
