[SkeletalMeshModelingTools] Adding a missing file from a previous change 52528792

#Synced-CL 52437958

diff --git a/Engine/Plugins/Animation/SkeletalMeshModelingTools/Source/SkeletalMeshModelingTools/Private/SkeletalMeshModelingToolsEditorMode.h b/Engine/Plugins/Animation/SkeletalMeshModelingTools/Source/SkeletalMeshModelingTools/Private/SkeletalMeshModelingToolsEditorMode.h
--- a/Engine/Plugins/Animation/SkeletalMeshModelingTools/Source/SkeletalMeshModelingTools/Private/SkeletalMeshModelingToolsEditorMode.h
+++ b/Engine/Plugins/Animation/SkeletalMeshModelingTools/Source/SkeletalMeshModelingTools/Private/SkeletalMeshModelingToolsEditorMode.h
@@ -283,7 +283,7 @@ private:
 	UPROPERTY(NonTransactional)
 	TObjectPtr<USkeletalMeshGeometrySelectionTransformTweaker> SelectionTransformTweaker;
 
-	bool bAllowGeometrySelection = true;
+	bool bSuspendGeometrySelection = false;
 
 	void InitializeGeometrySelectionSystems();
 
