[SkeletalMeshModelingTools] Fix CommitToSkeletalMesh: batch PostEditChange and fix FTrackerChangeScope redo semantics

Wrap commit with FScopedSkeletalMeshPostEditChange to batch multiple PostEditChange
notifications into one. Fix FTrackerChangeScope temporary that captured wrong redo state
by making it a named variable with ResetTrackersDirect inside its scope.

#Synced-CL 52437958

diff --git a/Engine/Plugins/Animation/SkeletalMeshModelingTools/Source/SkeletalMeshModelingTools/Private/Components/SKMBackedDynaMeshComponent.cpp b/Engine/Plugins/Animation/SkeletalMeshModelingTools/Source/SkeletalMeshModelingTools/Private/Components/SKMBackedDynaMeshComponent.cpp
--- a/Engine/Plugins/Animation/SkeletalMeshModelingTools/Source/SkeletalMeshModelingTools/Private/Components/SKMBackedDynaMeshComponent.cpp
+++ b/Engine/Plugins/Animation/SkeletalMeshModelingTools/Source/SkeletalMeshModelingTools/Private/Components/SKMBackedDynaMeshComponent.cpp
@@ -9,6 +9,7 @@
 #include "SkeletalMeshAttributes.h"
 #include "SkeletonModifier.h"
 #include "ScopedTransaction.h"
+#include "SkeletalMeshTypes.h"
 #include "Misc/ITransaction.h"
 #include "SkeletalMeshEditorSubsystem.h"
 #include "Parameterization/MeshPlanarSymmetry.h"
@@ -231,98 +232,105 @@ void USkeletalMeshBackedDynamicMeshComponent::CommitToSkeletalMesh()
 	
 	USkeletalMesh* SkeletalMesh = WeakSkeletalMesh.Get();
 
-	TArray<FName> OldEngineGeneratedMorphTargets;
-	USkeletalMeshEditorSubsystem::GetMorphTargetsGeneratedByEngine(SkeletalMesh, OldEngineGeneratedMorphTargets);
+	{
+		FScopedSkeletalMeshPostEditChange ScopedPostEditChange(SkeletalMesh);
+
+		TArray<FName> OldEngineGeneratedMorphTargets;
+		USkeletalMeshEditorSubsystem::GetMorphTargetsGeneratedByEngine(SkeletalMesh, OldEngineGeneratedMorphTargets);
 	
-	FScopedTransaction Transaction(LOCTEXT("ApplyChangesToAsset","Apply Changes To Skeletal Mesh"));
+		FScopedTransaction Transaction(LOCTEXT("ApplyChangesToAsset","Apply Changes To Skeletal Mesh"));
 
-	if (GetSkeletonChangeTracker().GetChangeCount() != 0)
-	{
-		USkeletonModifier* SkeletonModifier = NewObject<USkeletonModifier>();
-		SkeletonModifier->SetSkeletalMesh(SkeletalMesh);
+		if (GetSkeletonChangeTracker().GetChangeCount() != 0)
+		{
+			USkeletonModifier* SkeletonModifier = NewObject<USkeletonModifier>();
+			SkeletonModifier->SetSkeletalMesh(SkeletalMesh);
 
-		SkeletonModifier->ExternalUpdate(GetRefSkeleton(), GetSkeletonChangeTracker().GetBoneIndexTracker());
-		SkeletonModifier->CommitSkeletonToSkeletalMesh();
-	}
+			SkeletonModifier->ExternalUpdate(GetRefSkeleton(), GetSkeletonChangeTracker().GetBoneIndexTracker());
+			SkeletonModifier->CommitSkeletonToSkeletalMesh();
+		}
 
-	USkeletalMeshToolTargetFactory* LocalFactory = NewObject<USkeletalMeshToolTargetFactory>();
-	LocalFactory->SetActiveEditingLOD(GetLOD());
-	USkeletalMeshToolTarget* LocalTarget = CastChecked<USkeletalMeshToolTarget>(LocalFactory->BuildTarget(SkeletalMesh, {}));
+		USkeletalMeshToolTargetFactory* LocalFactory = NewObject<USkeletalMeshToolTargetFactory>();
+		LocalFactory->SetActiveEditingLOD(GetLOD());
+		USkeletalMeshToolTarget* LocalTarget = CastChecked<USkeletalMeshToolTarget>(LocalFactory->BuildTarget(SkeletalMesh, {}));
 
-	LocalTarget->CommitDynamicMesh(*GetMesh());
+		LocalTarget->CommitDynamicMesh(*GetMesh());
 
-	// Make sure morph targets are marked as engine generated correctly
-	// As soon as a morph is edited in engine, we want to mark it such that reimports in the future don't overwrite our edits.
-	{
-		TArray<FString> NewEngineGeneratedMorphTargets;
-	
-		for (const FName& OldEngineGeneratedMorphTarget : OldEngineGeneratedMorphTargets)
+		// Make sure morph targets are marked as engine generated correctly
+		// As soon as a morph is edited in engine, we want to mark it such that reimports in the future don't overwrite our edits.
 		{
-			FName CurrentName = GetMorphTargetChangeTracker().GetCurrentMorphTargetName(OldEngineGeneratedMorphTarget);
-			if (CurrentName != NAME_None)
+			TArray<FString> NewEngineGeneratedMorphTargets;
+	
+			for (const FName& OldEngineGeneratedMorphTarget : OldEngineGeneratedMorphTargets)
 			{
-				NewEngineGeneratedMorphTargets.AddUnique(CurrentName.ToString());
-				USkeletalMeshEditorSubsystem::SetMorphTargetsToGeneratedByEngine(SkeletalMesh, {CurrentName.ToString()});
+				FName CurrentName = GetMorphTargetChangeTracker().GetCurrentMorphTargetName(OldEngineGeneratedMorphTarget);
+				if (CurrentName != NAME_None)
+				{
+					NewEngineGeneratedMorphTargets.AddUnique(CurrentName.ToString());
+					USkeletalMeshEditorSubsystem::SetMorphTargetsToGeneratedByEngine(SkeletalMesh, {CurrentName.ToString()});
+				}
 			}
-		}
 	
-		for (const FName& EditedMorphTarget : GetMorphTargetChangeTracker().GetEditedMorphTargets())
-		{
-			NewEngineGeneratedMorphTargets.AddUnique(EditedMorphTarget.ToString());	
-		}
+			for (const FName& EditedMorphTarget : GetMorphTargetChangeTracker().GetEditedMorphTargets())
+			{
+				NewEngineGeneratedMorphTargets.AddUnique(EditedMorphTarget.ToString());	
+			}
 	
-		USkeletalMeshEditorSubsystem::SetMorphTargetsToGeneratedByEngine(SkeletalMesh, NewEngineGeneratedMorphTargets);
-	}
-
-	// Also change related curves
-	{
-		
-		UAnimCurveMetaData* AnimCurveMetaData = SkeletalMesh->GetAssetUserData<UAnimCurveMetaData>();
-		if (AnimCurveMetaData == nullptr)
-		{
-			AnimCurveMetaData = NewObject<UAnimCurveMetaData>(SkeletalMesh, NAME_None, RF_Transactional);
-			SkeletalMesh->AddAssetUserData(AnimCurveMetaData);
+			USkeletalMeshEditorSubsystem::SetMorphTargetsToGeneratedByEngine(SkeletalMesh, NewEngineGeneratedMorphTargets);
 		}
 
-		for (const FMorphTargetChangeTracker::FNameInfo& NameInfo : GetMorphTargetChangeTracker().GetCurvesToRename())
+		// Also change related curves
 		{
-			// Only rename if we have a morph flag set
-			if (FCurveMetaData* CurveMetaData = AnimCurveMetaData->GetCurveMetaData(NameInfo.OldName))
+		
+			UAnimCurveMetaData* AnimCurveMetaData = SkeletalMesh->GetAssetUserData<UAnimCurveMetaData>();
+			if (AnimCurveMetaData == nullptr)
+			{
+				AnimCurveMetaData = NewObject<UAnimCurveMetaData>(SkeletalMesh, NAME_None, RF_Transactional);
+				SkeletalMesh->AddAssetUserData(AnimCurveMetaData);
+			}
+
+			for (const FMorphTargetChangeTracker::FNameInfo& NameInfo : GetMorphTargetChangeTracker().GetCurvesToRename())
 			{
-				if (CurveMetaData->Type.bMorphtarget)
+				// Only rename if we have a morph flag set
+				if (FCurveMetaData* CurveMetaData = AnimCurveMetaData->GetCurveMetaData(NameInfo.OldName))
 				{
-					AnimCurveMetaData->RenameCurveMetaData(NameInfo.OldName, NameInfo.NewName);
-				}
-			}	
-		}
+					if (CurveMetaData->Type.bMorphtarget)
+					{
+						AnimCurveMetaData->RenameCurveMetaData(NameInfo.OldName, NameInfo.NewName);
+					}
+				}	
+			}
 	
-		for (const FName OldCurve : GetMorphTargetChangeTracker().GetCurvesToRemove())
-		{
-			// Only remove if we have a morph flag set
-			if (FCurveMetaData* CurveMetaData = AnimCurveMetaData->GetCurveMetaData(OldCurve))
+			for (const FName OldCurve : GetMorphTargetChangeTracker().GetCurvesToRemove())
 			{
-				if (CurveMetaData->Type.bMorphtarget)
+				// Only remove if we have a morph flag set
+				if (FCurveMetaData* CurveMetaData = AnimCurveMetaData->GetCurveMetaData(OldCurve))
 				{
-					AnimCurveMetaData->RemoveCurveMetaData(OldCurve);
-				}
-			}	
-		}
+					if (CurveMetaData->Type.bMorphtarget)
+					{
+						AnimCurveMetaData->RemoveCurveMetaData(OldCurve);
+					}
+				}	
+			}
 	
-		for (const FName NewCurve : GetMorphTargetChangeTracker().GetCurvesToAdd())
-		{
-			AnimCurveMetaData->AddCurveMetaData(NewCurve);
-		
-			// Ensure we have a morph flag set
-			if (FCurveMetaData* CurveMetaData = AnimCurveMetaData->GetCurveMetaData(NewCurve))
+			for (const FName NewCurve : GetMorphTargetChangeTracker().GetCurvesToAdd())
 			{
-				AnimCurveMetaData->Modify();
-				CurveMetaData->Type.bMorphtarget = true;
-			}	
+				AnimCurveMetaData->AddCurveMetaData(NewCurve);
+		
+				// Ensure we have a morph flag set
+				if (FCurveMetaData* CurveMetaData = AnimCurveMetaData->GetCurveMetaData(NewCurve))
+				{
+					AnimCurveMetaData->Modify();
+					CurveMetaData->Type.bMorphtarget = true;
+				}	
+			}
 		}
-	}
 
-	FTrackerChangeScope(this);
-	ResetTrackersDirect();
+		{
+			FTrackerChangeScope TrackerScope(this);
+			ResetTrackersDirect();
+		}
+		
+	} // FScopedSkeletalMeshPostEditChange exits, Skeletal Mesh build triggered
 	
 
 }
