[SkeletalMeshModelingTools] Relaunch morph editing tools on all morph target mutations

Extract RestartActiveMorphTargetEditingTool helper and call it from rename, remove,
duplicate, mirror, and flip handlers so the active tool picks up the latest data.

#Synced-CL 52437958

diff --git a/Engine/Plugins/Animation/SkeletalMeshModelingTools/Source/SkeletalMeshModelingTools/Private/SkeletalMeshModelingToolsEditorMode.cpp b/Engine/Plugins/Animation/SkeletalMeshModelingTools/Source/SkeletalMeshModelingTools/Private/SkeletalMeshModelingToolsEditorMode.cpp
--- a/Engine/Plugins/Animation/SkeletalMeshModelingTools/Source/SkeletalMeshModelingTools/Private/SkeletalMeshModelingToolsEditorMode.cpp
+++ b/Engine/Plugins/Animation/SkeletalMeshModelingTools/Source/SkeletalMeshModelingTools/Private/SkeletalMeshModelingToolsEditorMode.cpp
@@ -3,7 +3,6 @@
 #include "SkeletalMeshModelingToolsEditorMode.h"
 
 #include "AttributeEditorTool.h"
-#include "ToolInterfaces/IMorphTargetEditingToolInterface.h"
 #include "SkeletalMeshModelingToolsEditorModeToolkit.h"
 #include "SkeletalMeshModelingToolsCommands.h"
 #include "SkeletalMeshGizmoUtils.h"
@@ -639,28 +638,41 @@ FName USkeletalMeshModelingToolsEditorMode::GetEditingMorphTarget() const
 	return EditingMorphTarget;
 }
 
-
-
-void USkeletalMeshModelingToolsEditorMode::HandleSetEditingMorphTarget(FName InMorphTarget)
+USkeletalMeshModelingToolsEditorMode::FToolReactivateScope::FToolReactivateScope(USkeletalMeshModelingToolsEditorMode* InMode)
+	: Mode(InMode)
 {
-	// Such that EditingMorphTarget is recorded for undo/redo 
-	Modify();
-	EditingMorphTarget = InMorphTarget;
-
-	// Restart active morph target editing tool so it picks up the new editing target
-	UInteractiveToolManager* ToolManager = GetInteractiveToolsContext()->ToolManager;
-	if (ToolManager->HasActiveTool(EToolSide::Left))
+	if (Mode->ToolReactivateScopeDepth == 0)
 	{
-		UInteractiveTool* ActiveTool = ToolManager->GetActiveTool(EToolSide::Left);
-		if (Cast<IMorphTargetEditingToolInterface>(ActiveTool))
+		UInteractiveToolManager* ToolManager = Mode->GetInteractiveToolsContext()->ToolManager;
+		if (ToolManager->HasActiveTool(EToolSide::Left))
 		{
 			EToolShutdownType ShutdownType = ToolManager->CanAcceptActiveTool(EToolSide::Left)
 				? EToolShutdownType::Accept
 				: EToolShutdownType::Completed;
 			ToolManager->DeactivateTool(EToolSide::Left, ShutdownType);
-			ToolManager->ActivateTool(EToolSide::Left);
+			bNeedsReactivate = true;
 		}
 	}
+	Mode->ToolReactivateScopeDepth++;
+}
+
+USkeletalMeshModelingToolsEditorMode::FToolReactivateScope::~FToolReactivateScope()
+{
+	Mode->ToolReactivateScopeDepth--;
+	if (Mode->ToolReactivateScopeDepth == 0 && bNeedsReactivate)
+	{
+		UInteractiveToolManager* ToolManager = Mode->GetInteractiveToolsContext()->ToolManager;
+		ToolManager->ActivateTool(EToolSide::Left);
+	}
+}
+
+void USkeletalMeshModelingToolsEditorMode::HandleSetEditingMorphTarget(FName InMorphTarget)
+{
+	FToolReactivateScope ReactivateScope(this);
+
+	// Such that EditingMorphTarget is recorded for undo/redo 
+	Modify();
+	EditingMorphTarget = InMorphTarget;
 }
 
 
@@ -686,25 +698,25 @@ float USkeletalMeshModelingToolsEditorMode::GetMorphTargetWeight(FName MorphTarg
 
 FName USkeletalMeshModelingToolsEditorMode::HandleAddMorphTarget(FName InName)
 {
+	FToolReactivateScope ReactivateScope(this);
+
 	FName NewMorphTarget = GetCurrentEditingCache()->AddMorphTarget(InName);
+	GetCurrentEditingCache()->OverrideMorphTargetWeight(NewMorphTarget, 1.0f);
 	
 	HandleSetEditingMorphTarget(NewMorphTarget);
-	GetCurrentEditingCache()->OverrideMorphTargetWeight(NewMorphTarget, 1.0f);
 	
 	return NewMorphTarget;
 }
 
 FName USkeletalMeshModelingToolsEditorMode::HandleRenameMorphTarget(FName OldName, FName NewName)
 {
-	bool bSetEditingMorphTargetAgain = false;
-	if (GetEditingMorphTarget() == OldName)
-	{
-		bSetEditingMorphTargetAgain = true;
-	}
-	
+	FToolReactivateScope ReactivateScope(this);
+
+	bool bIsRenamingEditingMorphTarget = (GetEditingMorphTarget() == OldName);
+
 	FName FinalNewName = GetCurrentEditingCache()->RenameMorphTarget(OldName, NewName);
 
-	if (bSetEditingMorphTargetAgain)
+	if (bIsRenamingEditingMorphTarget)
 	{
 		HandleSetEditingMorphTarget(FinalNewName);
 	}
@@ -714,15 +726,13 @@ FName USkeletalMeshModelingToolsEditorMode::HandleRenameMorphTarget(FName OldNam
 
 void USkeletalMeshModelingToolsEditorMode::HandleRemoveMorphTargets(const TArray<FName>& Names)
 {
-	bool bClearEditingMorphTarget= false;
-	if (Names.Contains(GetEditingMorphTarget()))
-	{
-		bClearEditingMorphTarget = true;
-	}
-	
+	FToolReactivateScope ReactivateScope(this);
+
+	bool bIsRemovingEditingMorphTarget = Names.Contains(GetEditingMorphTarget());
+
 	GetCurrentEditingCache()->RemoveMorphTargets(Names);
 
-	if (bClearEditingMorphTarget)
+	if (bIsRemovingEditingMorphTarget)
 	{
 		HandleSetEditingMorphTarget(NAME_None);
 	}
@@ -730,16 +740,19 @@ void USkeletalMeshModelingToolsEditorMode::HandleRemoveMorphTargets(const TArray
 
 TArray<FName> USkeletalMeshModelingToolsEditorMode::HandleDuplicateMorphTargets(const TArray<FName>& Names)
 {
+	FToolReactivateScope ReactivateScope(this);
 	return GetCurrentEditingCache()->DuplicateMorphTargets(Names);
 }
 
 void USkeletalMeshModelingToolsEditorMode::HandleMirrorMorphTarget(FName Name)
 {
+	FToolReactivateScope ReactivateScope(this);
 	GetCurrentEditingCache()->MirrorMorphTargets({Name});
 }
 
 void USkeletalMeshModelingToolsEditorMode::HandleFlipMorphTarget(FName Name)
 {
+	FToolReactivateScope ReactivateScope(this);
 	GetCurrentEditingCache()->FlipMorphTargets({Name});
 }
 
diff --git a/Engine/Plugins/Animation/SkeletalMeshModelingTools/Source/SkeletalMeshModelingTools/Private/SkeletalMeshModelingToolsEditorMode.h b/Engine/Plugins/Animation/SkeletalMeshModelingTools/Source/SkeletalMeshModelingTools/Private/SkeletalMeshModelingToolsEditorMode.h
--- a/Engine/Plugins/Animation/SkeletalMeshModelingTools/Source/SkeletalMeshModelingTools/Private/SkeletalMeshModelingToolsEditorMode.h
+++ b/Engine/Plugins/Animation/SkeletalMeshModelingTools/Source/SkeletalMeshModelingTools/Private/SkeletalMeshModelingToolsEditorMode.h
@@ -310,4 +310,14 @@ private:
 
 	TValueWatcher<FSkeletonStateFactors> SkeletonStateUpdater;
 	
+	int32 ToolReactivateScopeDepth = 0;
+
+	struct FToolReactivateScope
+	{
+		FToolReactivateScope(USkeletalMeshModelingToolsEditorMode* InMode);
+		~FToolReactivateScope();
+	private:
+		USkeletalMeshModelingToolsEditorMode* Mode = nullptr;
+		bool bNeedsReactivate = false;
+	};
 };
diff --git a/Engine/Plugins/Animation/SkeletalMeshModelingTools/Source/SkeletalMeshModelingTools/Private/ToolInterfaces/IMorphTargetEditingToolInterface.cpp b/Engine/Plugins/Animation/SkeletalMeshModelingTools/Source/SkeletalMeshModelingTools/Private/ToolInterfaces/IMorphTargetEditingToolInterface.cpp
deleted file mode 100644
diff --git a/Engine/Plugins/Animation/SkeletalMeshModelingTools/Source/SkeletalMeshModelingTools/Public/ToolInterfaces/IMorphTargetEditingToolInterface.h b/Engine/Plugins/Animation/SkeletalMeshModelingTools/Source/SkeletalMeshModelingTools/Public/ToolInterfaces/IMorphTargetEditingToolInterface.h
deleted file mode 100644
diff --git a/Engine/Plugins/Experimental/Animation/SkeletalMeshMorphTargetEditingTools/Source/SkeletalMeshMorphTargetEditingTools/Private/IMorphTargetEditingToolInterface.cpp b/Engine/Plugins/Experimental/Animation/SkeletalMeshMorphTargetEditingTools/Source/SkeletalMeshMorphTargetEditingTools/Private/IMorphTargetEditingToolInterface.cpp
new file mode 100644
--- /dev/null
+++ b/Engine/Plugins/Experimental/Animation/SkeletalMeshMorphTargetEditingTools/Source/SkeletalMeshMorphTargetEditingTools/Private/IMorphTargetEditingToolInterface.cpp
@@ -0,0 +1,28 @@
+﻿// Copyright Epic Games, Inc. All Rights Reserved.
+
+#include "IMorphTargetEditingToolInterface.h"
+
+#include "ContextObjectStore.h"
+#include "InteractiveToolManager.h"
+#include "SkeletalMesh/SkeletalMeshEditingInterface.h"
+
+#include UE_INLINE_GENERATED_CPP_BY_NAME(IMorphTargetEditingToolInterface)
+
+void IMorphTargetEditingToolInterface::SetupMorphEditingToolCommon()
+{
+	UInteractiveTool* Tool = CastChecked<UInteractiveTool>(this);
+	USkeletalMeshEditorContextObjectBase* EditorContext = Tool->GetToolManager()->GetContextObjectStore()->FindContext<USkeletalMeshEditorContextObjectBase>();
+
+	
+	auto SetupFunction = [&](UMorphTargetEditingToolProperties* InProperties)
+	{
+		InProperties->EditMorphTargetName = EditorContext->GetEditingMorphTarget();
+	};
+	
+	SetupCommonProperties(SetupFunction);
+}
+
+void IMorphTargetEditingToolInterface::ShutdownMorphEditingToolCommon()
+{
+
+}
diff --git a/Engine/Plugins/Experimental/Animation/SkeletalMeshMorphTargetEditingTools/Source/SkeletalMeshMorphTargetEditingTools/Public/IMorphTargetEditingToolInterface.h b/Engine/Plugins/Experimental/Animation/SkeletalMeshMorphTargetEditingTools/Source/SkeletalMeshMorphTargetEditingTools/Public/IMorphTargetEditingToolInterface.h
new file mode 100644
--- /dev/null
+++ b/Engine/Plugins/Experimental/Animation/SkeletalMeshMorphTargetEditingTools/Source/SkeletalMeshMorphTargetEditingTools/Public/IMorphTargetEditingToolInterface.h
@@ -0,0 +1,35 @@
+﻿// Copyright Epic Games, Inc. All Rights Reserved.
+
+#pragma once
+
+#include "CoreMinimal.h"
+#include "InteractiveTool.h"
+#include "UObject/Interface.h"
+
+#include "IMorphTargetEditingToolInterface.generated.h"
+
+class UMorphTargetEditingToolProperties;
+
+UINTERFACE(MinimalAPI)
+class UMorphTargetEditingToolInterface : public UInterface
+{
+	GENERATED_BODY()
+};
+
+class IMorphTargetEditingToolInterface : public IInterface
+{
+	GENERATED_BODY()
+public:
+	virtual void SetupMorphEditingToolCommon();
+	virtual void ShutdownMorphEditingToolCommon();
+	virtual void SetupCommonProperties(const TFunction<void(UMorphTargetEditingToolProperties*)>& InSetupFunction) = 0;
+};
+
+UCLASS()
+class UMorphTargetEditingToolProperties : public UInteractiveToolPropertySet
+{
+	GENERATED_BODY()
+public:
+	UPROPERTY(VisibleAnywhere, Category = "MorphTarget", meta = (NoResetToDefault, DisplayName = "Name"))
+	FName EditMorphTargetName;
+};
diff --git a/Engine/Plugins/Experimental/Animation/SkeletalMeshMorphTargetEditingTools/Source/SkeletalMeshMorphTargetEditingTools/Public/MorphTargetMaskTool.h b/Engine/Plugins/Experimental/Animation/SkeletalMeshMorphTargetEditingTools/Source/SkeletalMeshMorphTargetEditingTools/Public/MorphTargetMaskTool.h
--- a/Engine/Plugins/Experimental/Animation/SkeletalMeshMorphTargetEditingTools/Source/SkeletalMeshMorphTargetEditingTools/Public/MorphTargetMaskTool.h
+++ b/Engine/Plugins/Experimental/Animation/SkeletalMeshMorphTargetEditingTools/Source/SkeletalMeshMorphTargetEditingTools/Public/MorphTargetMaskTool.h
@@ -2,7 +2,7 @@
 
 #pragma once
 
-#include "ToolInterfaces/IMorphTargetEditingToolInterface.h"
+#include "IMorphTargetEditingToolInterface.h"
 
 #include "MeshVertexAttributePaintToolBase.h"
 #include "SkeletalMesh/SkeletalMeshToolsHelper.h"
diff --git a/Engine/Plugins/Experimental/Animation/SkeletalMeshMorphTargetEditingTools/Source/SkeletalMeshMorphTargetEditingTools/Public/MorphTargetVertexSculptTool.h b/Engine/Plugins/Experimental/Animation/SkeletalMeshMorphTargetEditingTools/Source/SkeletalMeshMorphTargetEditingTools/Public/MorphTargetVertexSculptTool.h
--- a/Engine/Plugins/Experimental/Animation/SkeletalMeshMorphTargetEditingTools/Source/SkeletalMeshMorphTargetEditingTools/Public/MorphTargetVertexSculptTool.h
+++ b/Engine/Plugins/Experimental/Animation/SkeletalMeshMorphTargetEditingTools/Source/SkeletalMeshMorphTargetEditingTools/Public/MorphTargetVertexSculptTool.h
@@ -2,7 +2,7 @@
 
 #pragma once
 
-#include "ToolInterfaces/IMorphTargetEditingToolInterface.h"
+#include "IMorphTargetEditingToolInterface.h"
 #include "MeshVertexSculptTool.h"
 #include "SkeletalMesh/SkeletalMeshEditingInterface.h"
 #include "BaseTools/MultiSelectionMeshEditingTool.h"
