[SkeletalMeshModelingTools] Fix HandleSkeletalMeshChanged during undo/redo of commits

Add bExpectAssetChange consumable flag on editing component. Set during live commit
and in FTrackerChange::Revert/Apply for commit-tagged changes. Consumed by
HandleSkeletalMeshChanged to distinguish our own commit undo/redo from external changes.

#Synced-CL 52437958

diff --git a/Engine/Plugins/Animation/SkeletalMeshModelingTools/Source/SkeletalMeshModelingTools/Private/Components/SKMBackedDynaMeshComponent.cpp b/Engine/Plugins/Animation/SkeletalMeshModelingTools/Source/SkeletalMeshModelingTools/Private/Components/SKMBackedDynaMeshComponent.cpp
--- a/Engine/Plugins/Animation/SkeletalMeshModelingTools/Source/SkeletalMeshModelingTools/Private/Components/SKMBackedDynaMeshComponent.cpp
+++ b/Engine/Plugins/Animation/SkeletalMeshModelingTools/Source/SkeletalMeshModelingTools/Private/Components/SKMBackedDynaMeshComponent.cpp
@@ -223,12 +223,13 @@ void USkeletalMeshBackedDynamicMeshComponent::CommitToSkeletalMesh()
 		return;
 	}
 
-	if (bIsCommiting)
+	// If we are expecting a asset change, it means a commit is already in progress
+	if (bExpectAssetChange)
 	{
 		return;
 	}
 
-	TGuardValue<bool> CommitGuard(bIsCommiting, true);
+	bExpectAssetChange = true;
 	
 	USkeletalMesh* SkeletalMesh = WeakSkeletalMesh.Get();
 
@@ -327,6 +328,7 @@ void USkeletalMeshBackedDynamicMeshComponent::CommitToSkeletalMesh()
 
 		{
 			FTrackerChangeScope TrackerScope(this);
+			TrackerScope.MarkAsCommitChange();
 			ResetTrackersDirect();
 		}
 		
@@ -335,11 +337,6 @@ void USkeletalMeshBackedDynamicMeshComponent::CommitToSkeletalMesh()
 
 }
 
-bool USkeletalMeshBackedDynamicMeshComponent::IsCommitting() const
-{
-	return bIsCommiting;
-}
-
 void USkeletalMeshBackedDynamicMeshComponent::DiscardChanges()
 {
 	// Records a change that turns current mesh back to the initial state
@@ -892,6 +889,10 @@ void USkeletalMeshBackedDynamicMeshComponent::FTrackerChange::Apply(UObject* Obj
 {
 	if (USkeletalMeshBackedDynamicMeshComponent* Component = Cast<USkeletalMeshBackedDynamicMeshComponent>(Object))
 	{
+		if (bIsCommitChange)
+		{
+			Component->bExpectAssetChange = true;
+		}
 		Component->Trackers = NewTrackers;
 	}	
 }
@@ -900,10 +901,22 @@ void USkeletalMeshBackedDynamicMeshComponent::FTrackerChange::Revert(UObject* Ob
 {
 	if (USkeletalMeshBackedDynamicMeshComponent* Component = Cast<USkeletalMeshBackedDynamicMeshComponent>(Object))
 	{
+		if (bIsCommitChange)
+		{
+			Component->bExpectAssetChange = true;
+		}
 		Component->Trackers = OldTrackers;
 	}	
 }
 
+void USkeletalMeshBackedDynamicMeshComponent::FChangeScope::MarkAsCommitChange()
+{
+	if (TrackerChange)
+	{
+		TrackerChange->bIsCommitChange = true;
+	}
+}
+
 void USkeletalMeshBackedDynamicMeshComponent::FTrackerChange::Close(USkeletalMeshBackedDynamicMeshComponent* InComponent)
 {
 	NewTrackers = InComponent->Trackers;	
diff --git a/Engine/Plugins/Animation/SkeletalMeshModelingTools/Source/SkeletalMeshModelingTools/Private/Components/SKMBackedDynaMeshComponent.h b/Engine/Plugins/Animation/SkeletalMeshModelingTools/Source/SkeletalMeshModelingTools/Private/Components/SKMBackedDynaMeshComponent.h
--- a/Engine/Plugins/Animation/SkeletalMeshModelingTools/Source/SkeletalMeshModelingTools/Private/Components/SKMBackedDynaMeshComponent.h
+++ b/Engine/Plugins/Animation/SkeletalMeshModelingTools/Source/SkeletalMeshModelingTools/Private/Components/SKMBackedDynaMeshComponent.h
@@ -73,7 +73,7 @@ public:
 	EMeshLODIdentifier Init(USkeletalMesh* InSkeletalMesh, EMeshLODIdentifier InLOD);
 	EMeshLODIdentifier GetLOD() const;
 	void CommitToSkeletalMesh();
-	bool IsCommitting() const;
+
 	void DiscardChanges();
 
 	void MarkDirty();
@@ -111,7 +111,13 @@ public:
 	
 	DECLARE_MULTICAST_DELEGATE_OneParam(FOnRequestingVisibilityChange, bool);
 	FOnRequestingVisibilityChange& GetOnRequestingVisibilityChange() {return OnRequestingVisibilityChangeDelegate;}
-	
+
+	// Set before expected asset changes (commit, undo/redo of commit).
+	// Checked for re-entrancy in CommitToSkeletalMesh.
+	// Checked and cleared by HandleSkeletalMeshChanged to distinguish our own changes from external ones.
+	bool IsExpectingAssetChange() const { return bExpectAssetChange; }
+	void ClearExpectAssetChange() { bExpectAssetChange = false; }
+
 protected:
 	void ResetTrackersDirect();
 	FName GetAvailableMorphTargetName(FName InNewName, FName InOldName = NAME_None) const;
@@ -141,8 +147,8 @@ protected:
 	EMeshLODIdentifier LOD = EMeshLODIdentifier::LOD0;
 	int32 LODIndex = 0;
 
-	bool bIsCommiting = false;
-	
+	bool bExpectAssetChange = false;
+
 	FOnRequestingVisibilityChange OnRequestingVisibilityChangeDelegate;
 	
 	class FTrackerChange : public FToolCommandChange
@@ -150,6 +156,7 @@ protected:
 	public:
 		FTrackers OldTrackers;
 		FTrackers NewTrackers;
+		bool bIsCommitChange = false;
 		FTrackerChange(USkeletalMeshBackedDynamicMeshComponent* InComponent);
 		virtual FString ToString() const override;
 		virtual void Apply(UObject* Object) override;
@@ -163,7 +170,9 @@ protected:
 	public:
 		FChangeScope(USkeletalMeshBackedDynamicMeshComponent* InComponent, bool bInRecordMeshChange = true);
 		~FChangeScope();
-		
+
+		void MarkAsCommitChange();
+
 	private:
 		TUniquePtr<FTrackerChange> TrackerChange;
 		TSharedPtr<FDynamicMesh3> OldMesh;
diff --git a/Engine/Plugins/Animation/SkeletalMeshModelingTools/Source/SkeletalMeshModelingTools/Private/SkeletalMeshEditingCache.cpp b/Engine/Plugins/Animation/SkeletalMeshModelingTools/Source/SkeletalMeshModelingTools/Private/SkeletalMeshEditingCache.cpp
--- a/Engine/Plugins/Animation/SkeletalMeshModelingTools/Source/SkeletalMeshModelingTools/Private/SkeletalMeshEditingCache.cpp
+++ b/Engine/Plugins/Animation/SkeletalMeshModelingTools/Source/SkeletalMeshModelingTools/Private/SkeletalMeshEditingCache.cpp
@@ -301,11 +301,6 @@ void USkeletalMeshEditingCache::ApplyChanges()
 	}
 }
 
-bool USkeletalMeshEditingCache::IsApplyingChanges() const
-{
-	return EditingMeshComponent->IsCommitting();
-}
-
 void USkeletalMeshEditingCache::DiscardChanges()
 {
 	EditingMeshComponent->DiscardChanges();
diff --git a/Engine/Plugins/Animation/SkeletalMeshModelingTools/Source/SkeletalMeshModelingTools/Private/SkeletalMeshEditingCache.h b/Engine/Plugins/Animation/SkeletalMeshModelingTools/Source/SkeletalMeshModelingTools/Private/SkeletalMeshEditingCache.h
--- a/Engine/Plugins/Animation/SkeletalMeshModelingTools/Source/SkeletalMeshModelingTools/Private/SkeletalMeshEditingCache.h
+++ b/Engine/Plugins/Animation/SkeletalMeshModelingTools/Source/SkeletalMeshModelingTools/Private/SkeletalMeshEditingCache.h
@@ -58,7 +58,7 @@ public:
 	void Destroy();
 
 	void ApplyChanges();
-	bool IsApplyingChanges() const;
+
 	void DiscardChanges();
 
 	void Tick();
diff --git a/Engine/Plugins/Animation/SkeletalMeshModelingTools/Source/SkeletalMeshModelingTools/Private/SkeletalMeshModelingToolsEditorMode.cpp b/Engine/Plugins/Animation/SkeletalMeshModelingTools/Source/SkeletalMeshModelingTools/Private/SkeletalMeshModelingToolsEditorMode.cpp
--- a/Engine/Plugins/Animation/SkeletalMeshModelingTools/Source/SkeletalMeshModelingTools/Private/SkeletalMeshModelingToolsEditorMode.cpp
+++ b/Engine/Plugins/Animation/SkeletalMeshModelingTools/Source/SkeletalMeshModelingTools/Private/SkeletalMeshModelingToolsEditorMode.cpp
@@ -1791,21 +1791,28 @@ void USkeletalMeshModelingToolsEditorMode::HandleSkeletalMeshPreChange()
 
 void USkeletalMeshModelingToolsEditorMode::HandleSkeletalMeshChanged()
 {
-	if (GetCurrentEditingCache())
+	USkeletalMeshBackedDynamicMeshComponent* EditingMeshComponent = GetCurrentEditingCache() ? GetCurrentEditingCache()->GetEditingMeshComponent() : nullptr;
+	if (!EditingMeshComponent)
 	{
-		if (!GetCurrentEditingCache()->IsApplyingChanges())
-		{
-			int32 ChangeCount = GetCurrentEditingCache()->GetEditingMeshComponent()->GetChangeCount();
-			if (ChangeCount > 0)
-			{
-				FText Message = FText::Format(LOCTEXT("ExternalAssetChangeDiscardPendingToolChanges", "{0} change(s) discarded due to untracked external asset change"), ChangeCount);
-				GetToolManager()->DisplayMessage(Message, EToolMessageLevel::UserWarning);
-				ShowEditorMessage(ELogVerbosity::Type::Warning, Message);
-			}
-			
-			RecreateEditingCache(GetEditingLOD());
-		}
+		return;
 	}
+
+	if (EditingMeshComponent->IsExpectingAssetChange())
+	{
+		EditingMeshComponent->ClearExpectAssetChange();
+		return;
+	}
+
+	// Unexpected external change — nuke cache
+	int32 ChangeCount = EditingMeshComponent->GetChangeCount();
+	if (ChangeCount > 0)
+	{
+		FText Message = FText::Format(LOCTEXT("ExternalAssetChangeDiscardPendingToolChanges", "{0} change(s) discarded due to untracked external asset change"), ChangeCount);
+		GetToolManager()->DisplayMessage(Message, EToolMessageLevel::UserWarning);
+		ShowEditorMessage(ELogVerbosity::Type::Warning, Message);
+	}
+	
+	RecreateEditingCache(GetEditingLOD());
 }
 
 bool USkeletalMeshModelingToolsEditorMode::CanSetEditingLOD()
