[Backout] - CL52606759
#rnx
Original CL Desc
-----------------------------------------------------------------
[SkeletalMeshModelingTools] Fix HandleSkeletalMeshChanged during undo/redo of commits

Add bExpectAssetChange consumable flag on editing component. Set during live commit
and in FTrackerChange::Revert/Apply for commit-tagged changes. Consumed by
HandleSkeletalMeshChanged to distinguish our own commit undo/redo from external changes.

#Synced-CL 52437958

diff --git a/Engine/Plugins/Animation/SkeletalMeshModelingTools/Source/SkeletalMeshModelingTools/Private/Components/SKMBackedDynaMeshComponent.cpp b/Engine/Plugins/Animation/SkeletalMeshModelingTools/Source/SkeletalMeshModelingTools/Private/Components/SKMBackedDynaMeshComponent.cpp
--- a/Engine/Plugins/Animation/SkeletalMeshModelingTools/Source/SkeletalMeshModelingTools/Private/Components/SKMBackedDynaMeshComponent.cpp
+++ b/Engine/Plugins/Animation/SkeletalMeshModelingTools/Source/SkeletalMeshModelingTools/Private/Components/SKMBackedDynaMeshComponent.cpp
@@ -223,13 +223,12 @@ void USkeletalMeshBackedDynamicMeshComponent::CommitToSkeletalMesh()
 		return;
 	}
 
-	// If we are expecting a asset change, it means a commit is already in progress
-	if (bExpectAssetChange)
+	if (bIsCommiting)
 	{
 		return;
 	}
 
-	bExpectAssetChange = true;
+	TGuardValue<bool> CommitGuard(bIsCommiting, true);
 	
 	USkeletalMesh* SkeletalMesh = WeakSkeletalMesh.Get();
 
@@ -328,7 +327,6 @@ void USkeletalMeshBackedDynamicMeshComponent::CommitToSkeletalMesh()
 
 		{
 			FTrackerChangeScope TrackerScope(this);
-			TrackerScope.MarkAsCommitChange();
 			ResetTrackersDirect();
 		}
 		
@@ -337,6 +335,11 @@ void USkeletalMeshBackedDynamicMeshComponent::CommitToSkeletalMesh()
 
 }
 
+bool USkeletalMeshBackedDynamicMeshComponent::IsCommitting() const
+{
+	return bIsCommiting;
+}
+
 void USkeletalMeshBackedDynamicMeshComponent::DiscardChanges()
 {
 	// Records a change that turns current mesh back to the initial state
@@ -889,10 +892,6 @@ void USkeletalMeshBackedDynamicMeshComponent::FTrackerChange::Apply(UObject* Obj
 {
 	if (USkeletalMeshBackedDynamicMeshComponent* Component = Cast<USkeletalMeshBackedDynamicMeshComponent>(Object))
 	{
-		if (bIsCommitChange)
-		{
-			Component->bExpectAssetChange = true;
-		}
 		Component->Trackers = NewTrackers;
 	}	
 }
@@ -901,22 +900,10 @@ void USkeletalMeshBackedDynamicMeshComponent::FTrackerChange::Revert(UObject* Ob
 {
 	if (USkeletalMeshBackedDynamicMeshComponent* Component = Cast<USkeletalMeshBackedDynamicMeshComponent>(Object))
 	{
-		if (bIsCommitChange)
-		{
-			Component->bExpectAssetChange = true;
-		}
 		Component->Trackers = OldTrackers;
 	}	
 }
 
-void USkeletalMeshBackedDynamicMeshComponent::FChangeScope::MarkAsCommitChange()
-{
-	if (TrackerChange)
-	{
-		TrackerChange->bIsCommitChange = true;
-	}
-}
-
 void USkeletalMeshBackedDynamicMeshComponent::FTrackerChange::Close(USkeletalMeshBackedDynamicMeshComponent* InComponent)
 {
 	NewTrackers = InComponent->Trackers;	
diff --git a/Engine/Plugins/Animation/SkeletalMeshModelingTools/Source/SkeletalMeshModelingTools/Private/Components/SKMBackedDynaMeshComponent.h b/Engine/Plugins/Animation/SkeletalMeshModelingTools/Source/SkeletalMeshModelingTools/Private/Components/SKMBackedDynaMeshComponent.h
--- a/Engine/Plugins/Animation/SkeletalMeshModelingTools/Source/SkeletalMeshModelingTools/Private/Components/SKMBackedDynaMeshComponent.h
+++ b/Engine/Plugins/Animation/SkeletalMeshModelingTools/Source/SkeletalMeshModelingTools/Private/Components/SKMBackedDynaMeshComponent.h
@@ -73,7 +73,7 @@ public:
 	EMeshLODIdentifier Init(USkeletalMesh* InSkeletalMesh, EMeshLODIdentifier InLOD);
 	EMeshLODIdentifier GetLOD() const;
 	void CommitToSkeletalMesh();
-
+	bool IsCommitting() const;
 	void DiscardChanges();
 
 	void MarkDirty();
@@ -111,13 +111,7 @@ public:
 	
 	DECLARE_MULTICAST_DELEGATE_OneParam(FOnRequestingVisibilityChange, bool);
 	FOnRequestingVisibilityChange& GetOnRequestingVisibilityChange() {return OnRequestingVisibilityChangeDelegate;}
-
-	// Set before expected asset changes (commit, undo/redo of commit).
-	// Checked for re-entrancy in CommitToSkeletalMesh.
-	// Checked and cleared by HandleSkeletalMeshChanged to distinguish our own changes from external ones.
-	bool IsExpectingAssetChange() const { return bExpectAssetChange; }
-	void ClearExpectAssetChange() { bExpectAssetChange = false; }
-
+	
 protected:
 	void ResetTrackersDirect();
 	FName GetAvailableMorphTargetName(FName InNewName, FName InOldName = NAME_None) const;
@@ -147,8 +141,8 @@ protected:
 	EMeshLODIdentifier LOD = EMeshLODIdentifier::LOD0;
 	int32 LODIndex = 0;
 
-	bool bExpectAssetChange = false;
-
+	bool bIsCommiting = false;
+	
 	FOnRequestingVisibilityChange OnRequestingVisibilityChangeDelegate;
 	
 	class FTrackerChange : public FToolCommandChange
@@ -156,7 +150,6 @@ protected:
 	public:
 		FTrackers OldTrackers;
 		FTrackers NewTrackers;
-		bool bIsCommitChange = false;
 		FTrackerChange(USkeletalMeshBackedDynamicMeshComponent* InComponent);
 		virtual FString ToString() const override;
 		virtual void Apply(UObject* Object) override;
@@ -170,9 +163,7 @@ protected:
 	public:
 		FChangeScope(USkeletalMeshBackedDynamicMeshComponent* InComponent, bool bInRecordMeshChange = true);
 		~FChangeScope();
-
-		void MarkAsCommitChange();
-
+		
 	private:
 		TUniquePtr<FTrackerChange> TrackerChange;
 		TSharedPtr<FDynamicMesh3> OldMesh;
diff --git a/Engine/Plugins/Animation/SkeletalMeshModelingTools/Source/SkeletalMeshModelingTools/Private/SkeletalMeshEditingCache.cpp b/Engine/Plugins/Animation/SkeletalMeshModelingTools/Source/SkeletalMeshModelingTools/Private/SkeletalMeshEditingCache.cpp
--- a/Engine/Plugins/Animation/SkeletalMeshModelingTools/Source/SkeletalMeshModelingTools/Private/SkeletalMeshEditingCache.cpp
+++ b/Engine/Plugins/Animation/SkeletalMeshModelingTools/Source/SkeletalMeshModelingTools/Private/SkeletalMeshEditingCache.cpp
@@ -301,6 +301,11 @@ void USkeletalMeshEditingCache::ApplyChanges()
 	}
 }
 
+bool USkeletalMeshEditingCache::IsApplyingChanges() const
+{
+	return EditingMeshComponent->IsCommitting();
+}
+
 void USkeletalMeshEditingCache::DiscardChanges()
 {
 	EditingMeshComponent->DiscardChanges();
diff --git a/Engine/Plugins/Animation/SkeletalMeshModelingTools/Source/SkeletalMeshModelingTools/Private/SkeletalMeshEditingCache.h b/Engine/Plugins/Animation/SkeletalMeshModelingTools/Source/SkeletalMeshModelingTools/Private/SkeletalMeshEditingCache.h
--- a/Engine/Plugins/Animation/SkeletalMeshModelingTools/Source/SkeletalMeshModelingTools/Private/SkeletalMeshEditingCache.h
+++ b/Engine/Plugins/Animation/SkeletalMeshModelingTools/Source/SkeletalMeshModelingTools/Private/SkeletalMeshEditingCache.h
@@ -58,7 +58,7 @@ public:
 	void Destroy();
 
 	void ApplyChanges();
-
+	bool IsApplyingChanges() const;
 	void DiscardChanges();
 
 	void Tick();
diff --git a/Engine/Plugins/Animation/SkeletalMeshModelingTools/Source/SkeletalMeshModelingTools/Private/SkeletalMeshModelingToolsEditorMode.cpp b/Engine/Plugins/Animation/SkeletalMeshModelingTools/Source/SkeletalMeshModelingTools/Private/SkeletalMeshModelingToolsEditorMode.cpp
--- a/Engine/Plugins/Animation/SkeletalMeshModelingTools/Source/SkeletalMeshModelingTools/Private/SkeletalMeshModelingToolsEditorMode.cpp
+++ b/Engine/Plugins/Animation/SkeletalMeshModelingTools/Source/SkeletalMeshModelingTools/Private/SkeletalMeshModelingToolsEditorMode.cpp
@@ -1791,28 +1791,21 @@ void USkeletalMeshModelingToolsEditorMode::HandleSkeletalMeshPreChange()
 
 void USkeletalMeshModelingToolsEditorMode::HandleSkeletalMeshChanged()
 {
-	USkeletalMeshBackedDynamicMeshComponent* EditingMeshComponent = GetCurrentEditingCache() ? GetCurrentEditingCache()->GetEditingMeshComponent() : nullptr;
-	if (!EditingMeshComponent)
-	{
-		return;
-	}
-
-	if (EditingMeshComponent->IsExpectingAssetChange())
-	{
-		EditingMeshComponent->ClearExpectAssetChange();
-		return;
-	}
-
-	// Unexpected external change — nuke cache
-	int32 ChangeCount = EditingMeshComponent->GetChangeCount();
-	if (ChangeCount > 0)
+	if (GetCurrentEditingCache())
 	{
-		FText Message = FText::Format(LOCTEXT("ExternalAssetChangeDiscardPendingToolChanges", "{0} change(s) discarded due to untracked external asset change"), ChangeCount);
-		GetToolManager()->DisplayMessage(Message, EToolMessageLevel::UserWarning);
-		ShowEditorMessage(ELogVerbosity::Type::Warning, Message);
+		if (!GetCurrentEditingCache()->IsApplyingChanges())
+		{
+			int32 ChangeCount = GetCurrentEditingCache()->GetEditingMeshComponent()->GetChangeCount();
+			if (ChangeCount > 0)
+			{
+				FText Message = FText::Format(LOCTEXT("ExternalAssetChangeDiscardPendingToolChanges", "{0} change(s) discarded due to untracked external asset change"), ChangeCount);
+				GetToolManager()->DisplayMessage(Message, EToolMessageLevel::UserWarning);
+				ShowEditorMessage(ELogVerbosity::Type::Warning, Message);
+			}
+			
+			RecreateEditingCache(GetEditingLOD());
+		}
 	}
-	
-	RecreateEditingCache(GetEditingLOD());
 }
 
 bool USkeletalMeshModelingToolsEditorMode::CanSetEditingLOD()
