[SkeletalMeshEditor][SkeletalMeshModelingTools] Fix editor close flow for unapplied modeling tool changes

#Synced-CL 52607306

diff --git a/Engine/Plugins/Animation/SkeletalMeshModelingTools/Source/SkeletalMeshModelingTools/Private/SkeletalMeshModelingToolsEditorMode.cpp b/Engine/Plugins/Animation/SkeletalMeshModelingTools/Source/SkeletalMeshModelingTools/Private/SkeletalMeshModelingToolsEditorMode.cpp
--- a/Engine/Plugins/Animation/SkeletalMeshModelingTools/Source/SkeletalMeshModelingTools/Private/SkeletalMeshModelingToolsEditorMode.cpp
+++ b/Engine/Plugins/Animation/SkeletalMeshModelingTools/Source/SkeletalMeshModelingTools/Private/SkeletalMeshModelingToolsEditorMode.cpp
@@ -802,28 +802,8 @@ void USkeletalMeshModelingToolsEditorMode::Exit()
 	}
 
 	
-	if (GetCurrentEditingCache() )
+	if (GetCurrentEditingCache())
 	{
-		// Inform users if there are unapplied changes
-		int32 ChangeCount = GetCurrentEditingCache()->GetEditingMeshComponent()->GetChangeCount();
-		if (ChangeCount > 0 && GetSkeletalMesh())
-		{
-			EAppReturnType::Type YesNoReply = FMessageDialog::Open(EAppMsgType::YesNo,
-					FText::Format(
-						LOCTEXT("Prompt_SkeletalMeshModelingToolsClose", "Skeletal Mesh \"{0}\" has {1} unapplied changes. Would you like to apply them? (Selecting 'No' will cause all changes to be lost!)"),
-						FText::FromString(GetSkeletalMesh()->GetPathName()),
-						ChangeCount));
-
-			if (YesNoReply == EAppReturnType::Yes)
-			{
-				GetCurrentEditingCache()->ApplyChanges();
-			}
-			else
-			{
-				GetCurrentEditingCache()->DiscardChanges();
-			}
-		}
-
 		GetCurrentEditingCache()->Destroy();
 		CurrentEditingCache = nullptr;
 	}
@@ -1317,7 +1297,43 @@ bool USkeletalMeshModelingToolsEditorMode::HandleClick(FEditorViewportClient* In
 void USkeletalMeshModelingToolsEditorMode::PostUndo()
 {
 	Super::PostUndo();
+}
 
+bool USkeletalMeshModelingToolsEditorMode::OnRequestClose()
+{
+	if (GetCurrentEditingCache())
+	{
+		int32 ChangeCount = GetCurrentEditingCache()->GetEditingMeshComponent()->GetChangeCount();
+		if (ChangeCount > 0 && GetSkeletalMesh())
+		{
+			EAppReturnType::Type Reply = FMessageDialog::Open(EAppMsgType::YesNoCancel,
+				FText::Format(
+					LOCTEXT("Prompt_SkeletalMeshModelingToolsClose",
+						"Skeletal Mesh \"{0}\" has {1} unapplied {1}|plural(one=change,other=changes). Would you like to apply before closing?"),
+					FText::FromString(GetSkeletalMesh()->GetPathName()),
+					ChangeCount));
+
+			
+			switch (Reply)
+			{
+			case EAppReturnType::Yes:
+				{
+					FScopedTransaction Transaction(LOCTEXT("ApplyChangesToAssetOnModeExit", "Apply Changes To Skeletal Mesh"));
+					GetCurrentEditingCache()->ApplyChanges();
+					return true;
+				}
+			case EAppReturnType::No:
+				{
+					FScopedTransaction Transaction(LOCTEXT("DiscardChangesOnModeExit", "Discard Pending Changes"));
+					GetCurrentEditingCache()->DiscardChanges();
+					return true;
+				}
+			case EAppReturnType::Cancel:
+				return false;
+			}
+		}
+	}
+	return Super::OnRequestClose();
 }
 
 bool USkeletalMeshModelingToolsEditorMode::ComputeBoundingBoxForViewportFocus(AActor* Actor, UPrimitiveComponent* PrimitiveComponent, FBox& InOutBox) const
@@ -1738,8 +1754,6 @@ void USkeletalMeshModelingToolsEditorMode::SetEditorBinding(const TWeakPtr<ISkel
 	
 	EditorBinding = InSkeletalMeshEditor.Pin()->GetBinding();
 	Editor = InSkeletalMeshEditor;
-	Editor.Pin()->OnPreSaveAsset().AddUObject(this, &USkeletalMeshModelingToolsEditorMode::ApplyChanges);
-	Editor.Pin()->OnPreSaveAssetAs().AddUObject(this, &USkeletalMeshModelingToolsEditorMode::ApplyChanges);
 
 	SkeletalMesh = InSkeletalMeshEditor.Pin()->GetPersonaToolkit()->GetMesh();
 	SkeletalMesh->GetOnPreMeshChange().AddUObject(this, &USkeletalMeshModelingToolsEditorMode::HandleSkeletalMeshPreChange);
diff --git a/Engine/Plugins/Animation/SkeletalMeshModelingTools/Source/SkeletalMeshModelingTools/Private/SkeletalMeshModelingToolsEditorMode.h b/Engine/Plugins/Animation/SkeletalMeshModelingTools/Source/SkeletalMeshModelingTools/Private/SkeletalMeshModelingToolsEditorMode.h
--- a/Engine/Plugins/Animation/SkeletalMeshModelingTools/Source/SkeletalMeshModelingTools/Private/SkeletalMeshModelingToolsEditorMode.h
+++ b/Engine/Plugins/Animation/SkeletalMeshModelingTools/Source/SkeletalMeshModelingTools/Private/SkeletalMeshModelingToolsEditorMode.h
@@ -113,6 +113,7 @@ public:
 	virtual bool RequiresLegacyViewportInteractions() const override { return false; }
 
 	virtual void PostUndo() override;
+	virtual bool OnRequestClose() override;
 
 	// ISkeletalMeshBackedDynamicMeshComponentProvider
 	virtual USkeletalMeshBackedDynamicMeshComponent* GetComponent(UObject* SourceObject) override;
diff --git a/Engine/Plugins/Runtime/MeshModelingToolset/Source/MeshModelingToolsEditorOnly/Private/SkeletalMesh/SkinWeightsPaintTool.cpp b/Engine/Plugins/Runtime/MeshModelingToolset/Source/MeshModelingToolsEditorOnly/Private/SkeletalMesh/SkinWeightsPaintTool.cpp
--- a/Engine/Plugins/Runtime/MeshModelingToolset/Source/MeshModelingToolsEditorOnly/Private/SkeletalMesh/SkinWeightsPaintTool.cpp
+++ b/Engine/Plugins/Runtime/MeshModelingToolset/Source/MeshModelingToolsEditorOnly/Private/SkeletalMesh/SkinWeightsPaintTool.cpp
@@ -2311,13 +2311,6 @@ void USkinWeightsPaintTool::Shutdown(EToolShutdownType ShutdownType)
 {
 	SkeletonTreeNotifierBindScope.Reset();
 	
-	ISkeletalMeshBackedTarget* SkeletalMeshBackedTarget = Cast<ISkeletalMeshBackedTarget>(GetTarget()); 
-	USkeletalMesh* TargetSkeletalMesh = SkeletalMeshBackedTarget ? SkeletalMeshBackedTarget->GetSkeletalMesh() : nullptr;
-	if (TargetSkeletalMesh)
-	{
-		TargetSkeletalMesh->Modify();
-	}
-	
 	if (ShutdownType == EToolShutdownType::Accept)
 	{
 		// to force to refresh the tree
diff --git a/Engine/Source/Editor/SkeletalMeshEditor/Private/SkeletalMeshEditor.cpp b/Engine/Source/Editor/SkeletalMeshEditor/Private/SkeletalMeshEditor.cpp
--- a/Engine/Source/Editor/SkeletalMeshEditor/Private/SkeletalMeshEditor.cpp
+++ b/Engine/Source/Editor/SkeletalMeshEditor/Private/SkeletalMeshEditor.cpp
@@ -156,7 +156,11 @@ bool IsReductionParentBaseLODUseSkeletalMeshBuildWorkflow(USkeletalMesh* Skeleta
 
 bool FSkeletalMeshEditor::OnRequestClose(EAssetEditorCloseReason InCloseReason)
 {
-	bool bAllowClose = true;
+	bool bAllowClose = GetEditorModeManager().OnRequestClose();
+	if (!bAllowClose)
+	{
+		return false;
+	}
 
 	if (PersonaToolkit.IsValid() && SkeletalMesh)
 	{
@@ -218,6 +222,7 @@ bool FSkeletalMeshEditor::OnRequestClose(EAssetEditorCloseReason InCloseReason)
 			}
 		}
 	}
+
 	return bAllowClose;
 }
 
