[SkeletalMeshModelingTools] Isolate Geometry - Phase 1

#Synced-CL 52668932

diff --git a/Engine/Plugins/Animation/SkeletalMeshModelingTools/Source/SkeletalMeshModelingTools/Private/Components/SKMBackedDynaMeshComponent.cpp b/Engine/Plugins/Animation/SkeletalMeshModelingTools/Source/SkeletalMeshModelingTools/Private/Components/SKMBackedDynaMeshComponent.cpp
--- a/Engine/Plugins/Animation/SkeletalMeshModelingTools/Source/SkeletalMeshModelingTools/Private/Components/SKMBackedDynaMeshComponent.cpp
+++ b/Engine/Plugins/Animation/SkeletalMeshModelingTools/Source/SkeletalMeshModelingTools/Private/Components/SKMBackedDynaMeshComponent.cpp
@@ -3,6 +3,8 @@
 #include "Components/SKMBackedDynaMeshComponent.h"
 
 #include "DynamicMeshToMeshDescription.h"
+#include "DynamicSubmesh3.h"
+#include "GroupTopology.h"
 #include "LODUtilities.h"
 #include "MeshDescriptionToDynamicMesh.h"
 #include "ModelingToolTargetEditorOnlyUtil.h"
@@ -371,7 +373,8 @@ void USkeletalMeshBackedDynamicMeshComponent::HandleGeometryUpdate(
 	const FName EditingMorphTargetName,
 	const FDynamicMesh3& PosedMesh,
 	const TArray<FMatrix>& BoneMatrices,
-	const TMap<FName, float>& MorphTargetWeights)
+	const TMap<FName, float>& MorphTargetWeights,
+	const UE::Geometry::FDynamicSubmesh3* InSubmesh)
 {
 	using namespace UE::Geometry;
 	using namespace SkeletalMeshToolsHelper;
@@ -404,17 +407,24 @@ void USkeletalMeshBackedDynamicMeshComponent::HandleGeometryUpdate(
 		{
 			TFunction<void(FVertInfo, const FVector&)> WriteFunction;
 			TMap<FName, float> UnposeMorphTargetWeights = MorphTargetWeights;
-			
+
+			// Resolve vertex ID: map from submesh to full mesh space when isolated
+			auto ResolveVertID = [&InSubmesh](int32 VertID) -> int32
+			{
+				return InSubmesh ? InSubmesh->MapVertexToBaseMesh(VertID) : VertID;
+			};
+
 			if (MorphTargetAttribute)
 			{
 				WriteFunction = [&](FVertInfo VertInfo, const FVector& UnposedVertPos)
 				{
-					int32 VertID = VertInfo.VertID;
-
-					FVector NewDelta = UnposedVertPos - FVector(Mesh.GetVertex(VertID));
-					NewDelta /= EditingMorphTargetWeight;
-
-					MorphTargetAttribute->SetValue(VertID, NewDelta);
+					const int32 VertID = ResolveVertID(VertInfo.VertID);
+					if (VertID != FDynamicMesh3::InvalidID)
+					{
+						FVector NewDelta = UnposedVertPos - FVector(Mesh.GetVertex(VertID));
+						NewDelta /= EditingMorphTargetWeight;
+						MorphTargetAttribute->SetValue(VertID, NewDelta);
+					}
 				};
 
 				// Exclude the morph target we are trying to extract
@@ -424,13 +434,17 @@ void USkeletalMeshBackedDynamicMeshComponent::HandleGeometryUpdate(
 			{
 				WriteFunction = [&](FVertInfo VertInfo, const FVector& UnposedVertPos)
 				{
-					int32 VertID = VertInfo.VertID;
-
-					Mesh.SetVertex(VertID, UnposedVertPos);
+					const int32 VertID = ResolveVertID(VertInfo.VertID);
+					if (VertID != FDynamicMesh3::InvalidID)
+					{
+						Mesh.SetVertex(VertID, UnposedVertPos);
+					}
 				};
-			}		
+			}
 
-			GetUnposedMesh(WriteFunction, PosedMesh, Mesh, BoneMatrices, NAME_None, UnposeMorphTargetWeights);
+			// When isolated, use the submesh as SourceMesh (it has bone weights copied by FDynamicSubmesh3)
+			const FDynamicMesh3& SourceMesh = InSubmesh ? InSubmesh->GetSubmesh() : Mesh;
+			GetUnposedMesh(WriteFunction, PosedMesh, SourceMesh, BoneMatrices, NAME_None, UnposeMorphTargetWeights);
 		});
 			
 	// Tracker change
@@ -960,6 +974,28 @@ USkeletalMeshBackedDynamicMeshComponent::FChangeScope::~FChangeScope()
 	}
 }
 
+const UE::Geometry::FGroupTopology& USkeletalMeshBackedDynamicMeshComponent::GetGroupTopology()
+{
+	UpdateGroupTopologyIfNeeded();
+	check(CachedGroupTopology);
+	return *CachedGroupTopology;
+}
+
+void USkeletalMeshBackedDynamicMeshComponent::UpdateGroupTopologyIfNeeded()
+{
+	if (BaseGeometryVersionForCurrentGroupTopology == Trackers.BaseGeometryChangeCount)
+	{
+		return;
+	}
+
+	BaseGeometryVersionForCurrentGroupTopology = Trackers.BaseGeometryChangeCount;
+
+	GetDynamicMesh()->ProcessMesh([this](const FDynamicMesh3& Mesh)
+	{
+		CachedGroupTopology = MakePimpl<UE::Geometry::FGroupTopology>(&Mesh, true);
+	});
+}
+
 void USkeletalMeshBackedDynamicMeshComponent::UpdateSymmetryIfNeeded()
 {
 	using namespace UE::Geometry;
diff --git a/Engine/Plugins/Animation/SkeletalMeshModelingTools/Source/SkeletalMeshModelingTools/Private/Components/SKMBackedDynaMeshComponent.h b/Engine/Plugins/Animation/SkeletalMeshModelingTools/Source/SkeletalMeshModelingTools/Private/Components/SKMBackedDynaMeshComponent.h
--- a/Engine/Plugins/Animation/SkeletalMeshModelingTools/Source/SkeletalMeshModelingTools/Private/Components/SKMBackedDynaMeshComponent.h
+++ b/Engine/Plugins/Animation/SkeletalMeshModelingTools/Source/SkeletalMeshModelingTools/Private/Components/SKMBackedDynaMeshComponent.h
@@ -15,6 +15,8 @@
 namespace UE::Geometry
 {
 	class FMeshPlanarSymmetry;
+	struct FDynamicSubmesh3;
+	class FGroupTopology;
 }
 
 class USkeletonModifier;
@@ -85,7 +87,8 @@ public:
 		const FName EditingMorphTargetName,
 		const FDynamicMesh3& PosedMesh,
 		const TArray<FMatrix>& BoneMatrices,
-		const TMap<FName, float>& MorphTargetWeights
+		const TMap<FName, float>& MorphTargetWeights,
+		const UE::Geometry::FDynamicSubmesh3* InSubmesh = nullptr
 		);
 
 	FName AddMorphTarget(FName InName);
@@ -108,6 +111,8 @@ public:
 	
 
 	USkeletalMesh* GetSkeletalMesh() const;
+
+	const UE::Geometry::FGroupTopology& GetGroupTopology();
 	
 	DECLARE_MULTICAST_DELEGATE_OneParam(FOnRequestingVisibilityChange, bool);
 	FOnRequestingVisibilityChange& GetOnRequestingVisibilityChange() {return OnRequestingVisibilityChangeDelegate;}
@@ -199,6 +204,11 @@ protected:
 	TPimplPtr<UE::Geometry::FMeshPlanarSymmetry> ReadOnlySymmetry;
 	int32 BaseGeometryVersionForCurrentSymmetry = -1;
 	void UpdateSymmetryIfNeeded();
+
+	// Cached group topology, lazily built and invalidated when base geometry changes
+	TPimplPtr<UE::Geometry::FGroupTopology> CachedGroupTopology;
+	int32 BaseGeometryVersionForCurrentGroupTopology = -1;
+	void UpdateGroupTopologyIfNeeded();
 };
 
 #undef UE_API
diff --git a/Engine/Plugins/Animation/SkeletalMeshModelingTools/Source/SkeletalMeshModelingTools/Private/Selection/SKMEditingCacheSelector.cpp b/Engine/Plugins/Animation/SkeletalMeshModelingTools/Source/SkeletalMeshModelingTools/Private/Selection/SKMEditingCacheSelector.cpp
new file mode 100644
--- /dev/null
+++ b/Engine/Plugins/Animation/SkeletalMeshModelingTools/Source/SkeletalMeshModelingTools/Private/Selection/SKMEditingCacheSelector.cpp
@@ -0,0 +1,894 @@
+// Copyright Epic Games, Inc. All Rights Reserved.
+
+#include "Selection/SKMEditingCacheSelector.h"
+
+#include "SkeletalMeshEditingCache.h"
+#include "SkeletalMeshModelingToolsEditorMode.h"
+#include "Components/SKMBackedDynaMeshComponent.h"
+#include "DynamicSubmesh3.h"
+#include "GroupTopology.h"
+#include "Selections/GeometrySelectionUtil.h"
+#include "Selection/ToolSelectionUtil.h"
+
+using namespace UE::Geometry;
+
+
+// ============================================================================
+// Selection remapping helper (editing-mesh IDs -> preview-mesh IDs)
+// ============================================================================
+
+static FGeometrySelection RemapSelectionToPreviewMesh(const FGeometrySelection& EditingSelection, const FDynamicSubmesh3* SubmeshMapping, const FGroupTopology* PreviewGroupTopology = nullptr)
+{
+	if (!SubmeshMapping)
+	{
+		return EditingSelection;
+	}
+
+	FGeometrySelection PreviewSelection;
+	PreviewSelection.ElementType = EditingSelection.ElementType;
+	PreviewSelection.TopologyType = EditingSelection.TopologyType;
+	PreviewSelection.Selection.Reserve(EditingSelection.Selection.Num());
+
+	for (uint64 EncodedID : EditingSelection.Selection)
+	{
+		FGeoSelectionID SelID(EncodedID);
+		FGeoSelectionID PreviewSelID;
+		bool bValid = false;
+
+		if (EditingSelection.TopologyType == EGeometryTopologyType::Triangle)
+		{
+			switch (EditingSelection.ElementType)
+			{
+			case EGeometryElementType::Face:
+			{
+				int32 PreviewTID = SubmeshMapping->MapTriangleToSubmesh(static_cast<int32>(SelID.GeometryID));
+				if (PreviewTID != FDynamicMesh3::InvalidID) { PreviewSelID = FGeoSelectionID::MeshTriangle(PreviewTID); bValid = true; }
+				break;
+			}
+			case EGeometryElementType::Edge:
+			{
+				FMeshTriEdgeID TriEdge(SelID.GeometryID);
+				int32 PreviewTID = SubmeshMapping->MapTriangleToSubmesh(static_cast<int32>(TriEdge.TriangleID));
+				if (PreviewTID != FDynamicMesh3::InvalidID) { PreviewSelID = FGeoSelectionID::MeshEdge(FMeshTriEdgeID(PreviewTID, TriEdge.TriEdgeIndex)); bValid = true; }
+				break;
+			}
+			case EGeometryElementType::Vertex:
+			{
+				int32 PreviewVID = SubmeshMapping->MapVertexToSubmesh(static_cast<int32>(SelID.GeometryID));
+				if (PreviewVID != FDynamicMesh3::InvalidID) { PreviewSelID = FGeoSelectionID::MeshVertex(PreviewVID); bValid = true; }
+				break;
+			}
+			}
+		}
+		else if (EditingSelection.TopologyType == EGeometryTopologyType::Polygroup)
+		{
+			switch (EditingSelection.ElementType)
+			{
+			case EGeometryElementType::Face:
+			{
+				int32 PreviewTID = SubmeshMapping->MapTriangleToSubmesh(static_cast<int32>(SelID.GeometryID));
+				int32 PreviewGID = SubmeshMapping->MapGroupToSubmesh(static_cast<int32>(SelID.TopologyID));
+				if (PreviewTID != FDynamicMesh3::InvalidID && PreviewGID != FDynamicMesh3::InvalidID)
+				{
+					PreviewSelID = FGeoSelectionID::GroupFace(PreviewTID, PreviewGID);
+					bValid = true;
+				}
+				break;
+			}
+			case EGeometryElementType::Edge:
+			{
+				if (PreviewGroupTopology)
+				{
+					FMeshTriEdgeID EditingTriEdge(SelID.GeometryID);
+					const FDynamicMesh3* BaseMesh = SubmeshMapping->GetBaseMesh();
+					int32 EditingEdgeID = BaseMesh->GetTriEdge(EditingTriEdge.TriangleID, EditingTriEdge.TriEdgeIndex);
+					int32 PreviewEdgeID = SubmeshMapping->MapEdgeToSubmesh(EditingEdgeID);
+					if (PreviewEdgeID != FDynamicMesh3::InvalidID)
+					{
+						int32 PreviewGroupEdgeID = PreviewGroupTopology->FindGroupEdgeID(PreviewEdgeID);
+						FMeshTriEdgeID PreviewTriEdge = SubmeshMapping->GetSubmesh().GetTriEdgeIDFromEdgeID(PreviewEdgeID);
+						if (PreviewGroupEdgeID >= 0 && PreviewTriEdge.TriangleID != FDynamicMesh3::InvalidID)
+						{
+							PreviewSelID = FGeoSelectionID(PreviewTriEdge.Encoded(), PreviewGroupEdgeID);
+							bValid = true;
+						}
+					}
+				}
+				break;
+			}
+			case EGeometryElementType::Vertex:
+			{
+				int32 PreviewVID = SubmeshMapping->MapVertexToSubmesh(static_cast<int32>(SelID.GeometryID));
+				if (PreviewVID != FDynamicMesh3::InvalidID)
+				{
+					// CornerID differs between topologies — skip TopologyID, let consumer re-derive
+					PreviewSelID = FGeoSelectionID(static_cast<uint32>(PreviewVID), 0);
+					bValid = true;
+				}
+				break;
+			}
+			}
+		}
+
+		if (bValid)
+		{
+			PreviewSelection.Selection.Add(PreviewSelID.Encoded());
+		}
+	}
+
+	return PreviewSelection;
+}
+
+
+/**
+ * Remap encoded selection element IDs from preview-mesh space to editing-mesh space.
+ * Handles all topology/element type combinations.
+ */
+static void RemapPreviewIDsToEditingMesh(
+	TArray<uint64>& Elements,
+	EGeometryTopologyType TopologyType,
+	EGeometryElementType ElementType,
+	const FDynamicSubmesh3* SubmeshMapping,
+	UDynamicMesh* PreviewMesh,
+	UDynamicMesh* InEditingMesh,
+	const FGroupTopology* EditingTopology)
+{
+	for (uint64& Encoded : Elements)
+	{
+		FGeoSelectionID SelID(Encoded);
+
+		if (TopologyType == EGeometryTopologyType::Triangle)
+		{
+			if (ElementType == EGeometryElementType::Face)
+			{
+				int32 EditingTID = SubmeshMapping->MapTriangleToBaseMesh(static_cast<int32>(SelID.GeometryID));
+				Encoded = FGeoSelectionID::MeshTriangle(EditingTID).Encoded();
+			}
+			else if (ElementType == EGeometryElementType::Edge)
+			{
+				FMeshTriEdgeID PreviewTriEdge(SelID.GeometryID);
+				int32 EditingTID = SubmeshMapping->MapTriangleToBaseMesh(static_cast<int32>(PreviewTriEdge.TriangleID));
+				Encoded = FGeoSelectionID::MeshEdge(FMeshTriEdgeID(EditingTID, PreviewTriEdge.TriEdgeIndex)).Encoded();
+			}
+			else if (ElementType == EGeometryElementType::Vertex)
+			{
+				int32 EditingVID = SubmeshMapping->MapVertexToBaseMesh(static_cast<int32>(SelID.GeometryID));
+				Encoded = FGeoSelectionID::MeshVertex(EditingVID).Encoded();
+			}
+		}
+		else if (TopologyType == EGeometryTopologyType::Polygroup)
+		{
+			if (ElementType == EGeometryElementType::Face)
+			{
+				int32 EditingTID = SubmeshMapping->MapTriangleToBaseMesh(static_cast<int32>(SelID.GeometryID));
+				int32 EditingGroupID = SubmeshMapping->MapGroupToBaseMesh(static_cast<int32>(SelID.TopologyID));
+				Encoded = FGeoSelectionID(static_cast<uint32>(EditingTID), static_cast<uint32>(EditingGroupID)).Encoded();
+			}
+			else if (ElementType == EGeometryElementType::Edge)
+			{
+				FMeshTriEdgeID PreviewTriEdge(SelID.GeometryID);
+				int32 PreviewEdgeID = -1;
+				PreviewMesh->ProcessMesh([&](const FDynamicMesh3& Mesh)
+				{
+					PreviewEdgeID = Mesh.IsTriangle(PreviewTriEdge.TriangleID) ?
+						Mesh.GetTriEdge(PreviewTriEdge.TriangleID, PreviewTriEdge.TriEdgeIndex) : IndexConstants::InvalidID;
+				});
+				int32 EditingEID = SubmeshMapping->MapEdgeToBaseMesh(PreviewEdgeID);
+				if (EditingEID != FDynamicMesh3::InvalidID && EditingTopology)
+				{
+					int32 EditingGroupEdgeID = EditingTopology->FindGroupEdgeID(EditingEID);
+					FMeshTriEdgeID EditingTriEdgeID;
+					InEditingMesh->ProcessMesh([&](const FDynamicMesh3& Mesh) { EditingTriEdgeID = Mesh.GetTriEdgeIDFromEdgeID(EditingEID); });
+					Encoded = FGeoSelectionID(EditingTriEdgeID.Encoded(), static_cast<uint32>(EditingGroupEdgeID)).Encoded();
+				}
+			}
+			else if (ElementType == EGeometryElementType::Vertex)
+			{
+				int32 EditingVID = SubmeshMapping->MapVertexToBaseMesh(static_cast<int32>(SelID.GeometryID));
+				int32 EditingCornerID = EditingTopology ? EditingTopology->GetCornerIDFromVertexID(EditingVID) : 0;
+				Encoded = FGeoSelectionID(static_cast<uint32>(EditingVID), static_cast<uint32>(EditingCornerID)).Encoded();
+			}
+		}
+	}
+}
+
+
+// ============================================================================
+// Custom Transformers (remap selection in BeginTransform, then delegate to base)
+// ============================================================================
+
+void FSkeletalMeshEditingCacheSelectionTransformer::BeginTransform(const FGeometrySelection& Selection)
+{
+	FGeometrySelection PreviewSelection = RemapSelectionToPreviewMesh(Selection, IsolationSubmeshMapping, PreviewGroupTopology);
+	FBasicDynamicMeshSelectionTransformer::BeginTransform(PreviewSelection);
+}
+
+void FSkeletalMeshEditingCachePolygroupTransformer::BeginTransform(const FGeometrySelection& Selection)
+{
+	FGeometrySelection PreviewSelection = RemapSelectionToPreviewMesh(Selection, IsolationSubmeshMapping, PreviewGroupTopology);
+	FDynamicMeshPolygroupTransformer::BeginTransform(PreviewSelection);
+}
+
+
+// ============================================================================
+// Lifecycle
+// ============================================================================
+
+FSkeletalMeshEditingCacheSelector::~FSkeletalMeshEditingCacheSelector() = default;
+
+void FSkeletalMeshEditingCacheSelector::InitializeFromEditingCache(
+	FGeometryIdentifier InSourceGeometryIdentifier,
+	USkeletalMeshEditingCache* InEditingCache,
+	TFunction<FName()> InGetEditingMorphTargetFunc)
+{
+	check(InEditingCache);
+
+	WeakEditingCache = InEditingCache;
+	GetEditingMorphTargetFunc = MoveTemp(InGetEditingMorphTargetFunc);
+
+	// Store editing mesh reference (for topology/selection ID queries)
+	USkeletalMeshBackedDynamicMeshComponent* EditingMeshComponent = InEditingCache->GetEditingMeshComponent();
+	EditingMesh = EditingMeshComponent->GetDynamicMesh();
+	EditingMeshIdentifier = FGeometryIdentifier::PrimitiveComponent(EditingMeshComponent, FGeometryIdentifier::EObjectType::DynamicMeshComponent);
+
+	// Listen for editing mesh changes to invalidate editing topology cache
+	EditingMesh_OnMeshChangedHandle = EditingMesh->OnMeshChanged().AddLambda(
+		[this](UDynamicMesh*, FDynamicMeshChangeInfo)
+		{
+			InvalidateEditingMeshCaches();
+		});
+
+	// Base class: TargetMesh = PREVIEW mesh (for spatial queries, transformers, rendering)
+	UDynamicMeshComponent* PreviewComponent = CastChecked<UDynamicMeshComponent>(InEditingCache->GetPreviewMeshComponent());
+	TWeakObjectPtr<USkeletalMeshEditingCache> WeakCache = InEditingCache;
+	Initialize(InSourceGeometryIdentifier, PreviewComponent->GetDynamicMesh(),
+		[WeakCache]() -> FTransformSRT3d
+		{
+			return WeakCache.IsValid() ? FTransformSRT3d(WeakCache->GetTransform()) : FTransformSRT3d::Identity();
+		});
+
+	// Subscribe to isolation changes
+	OnIsolationChangedHandle = InEditingCache->GetOnIsolationChanged().AddRaw(this, &FSkeletalMeshEditingCacheSelector::HandleIsolationChanged);
+
+	// Set initial isolation state
+	if (InEditingCache->HasIsolation())
+	{
+		SetActiveSubmeshMapping(InEditingCache->GetIsolationSubmesh());
+	}
+}
+
+void FSkeletalMeshEditingCacheSelector::Shutdown()
+{
+	if (USkeletalMeshEditingCache* Cache = WeakEditingCache.Get())
+	{
+		Cache->GetOnIsolationChanged().Remove(OnIsolationChangedHandle);
+		OnIsolationChangedHandle.Reset();
+	}
+
+	if (EditingMesh.IsValid() && EditingMesh_OnMeshChangedHandle.IsValid())
+	{
+		EditingMesh->OnMeshChanged().Remove(EditingMesh_OnMeshChangedHandle);
+		EditingMesh_OnMeshChangedHandle.Reset();
+	}
+
+	EditingGroupTopology.Reset();
+	EditingMesh = nullptr;
+	ActiveSubmeshMapping = nullptr;
+
+	FBaseDynamicMeshSelector::Shutdown();
+}
+
+void FSkeletalMeshEditingCacheSelector::SetActiveSubmeshMapping(const FDynamicSubmesh3* InSubmeshMapping)
+{
+	ActiveSubmeshMapping = InSubmeshMapping;
+}
+
+
+
+// ============================================================================
+// Editing-mesh topology overrides
+// ============================================================================
+
+void FSkeletalMeshEditingCacheSelector::InvalidateEditingMeshCaches()
+{
+	EditingGroupTopology.Reset();
+}
+
+const FGroupTopology* FSkeletalMeshEditingCacheSelector::GetEditingMeshGroupTopology()
+{
+	if (!EditingGroupTopology.IsValid() && EditingMesh.IsValid())
+	{
+		EditingGroupTopology = MakePimpl<FGroupTopology>(EditingMesh->GetMeshPtr(), true);
+	}
+	return EditingGroupTopology.Get();
+}
+
+void FSkeletalMeshEditingCacheSelector::InitializeSelectionFromPredicate(
+	FGeometrySelection& SelectionInOut,
+	TFunctionRef<bool(FGeoSelectionID)> SelectionIDPredicate,
+	EInitializeSelectionMode InitializeMode,
+	const FGeometrySelection* ReferenceSelection)
+{
+	if (!ActiveSubmeshMapping)
+	{
+		// No isolation — IDs match, delegate to base
+		FBaseDynamicMeshSelector::InitializeSelectionFromPredicate(SelectionInOut, SelectionIDPredicate, InitializeMode, ReferenceSelection);
+		return;
+	}
+
+	if (InitializeMode != EInitializeSelectionMode::All && !ensure(ReferenceSelection != nullptr))
+	{
+		return;
+	}
+
+	// With isolation, must construct selection IDs from the editing mesh
+	EditingMesh->ProcessMesh([&](const FDynamicMesh3& Mesh)
+	{
+		const FGroupTopology* UseTopology = (SelectionInOut.TopologyType == EGeometryTopologyType::Polygroup) ? GetEditingMeshGroupTopology() : nullptr;
+		if (InitializeMode == EInitializeSelectionMode::Connected)
+		{
+			MakeSelectAllConnectedSelection(Mesh, UseTopology, *ReferenceSelection,
+				SelectionIDPredicate, [](FGeoSelectionID, FGeoSelectionID) { return true; }, SelectionInOut);
+		}
+		else if (InitializeMode == EInitializeSelectionMode::AdjacentToBorder)
+		{
+			MakeBoundaryConnectedSelection(Mesh, UseTopology, *ReferenceSelection,
+				SelectionIDPredicate, SelectionInOut);
+		}
+		else // All
+		{
+			MakeSelectAllSelection(Mesh, UseTopology, SelectionIDPredicate, SelectionInOut);
+		}
+	});
+}
+
+void FSkeletalMeshEditingCacheSelector::UpdateSelectionFromSelection(
+	const FGeometrySelection& FromSelection,
+	bool bAllowConversion,
+	FGeometrySelectionEditor& SelectionEditor,
+	const FGeometrySelectionUpdateConfig& UpdateConfig,
+	FGeometrySelectionDelta* SelectionDelta)
+{
+	// Selection IDs are always in editing-mesh space, so this logic is the same regardless of isolation.
+	// The only difference from the base class is we use GetEditingMeshGroupTopology() for conversion.
+	if (IsLockable() && IsLocked())
+	{
+		return;
+	}
+
+	if (FromSelection.IsSameType(SelectionEditor.GetSelection()))
+	{
+		UpdateSelectionWithNewElements(&SelectionEditor, UpdateConfig.ChangeType, FromSelection.Selection.Array(), SelectionDelta);
+		return;
+	}
+
+	if (!bAllowConversion)
+	{
+		return;
+	}
+
+	EditingMesh->ProcessMesh([&](const FDynamicMesh3& Mesh)
+	{
+		FGeometrySelection TempSelection;
+		TempSelection.InitializeTypes(SelectionEditor.GetElementType(), SelectionEditor.GetTopologyType());
+		bool bConverted = UE::Geometry::ConvertSelection(Mesh, GetEditingMeshGroupTopology(), FromSelection, TempSelection,
+			EEnumerateSelectionConversionParams::ContainSelection);
+		if (bConverted)
+		{
+			UpdateSelectionWithNewElements(&SelectionEditor, UpdateConfig.ChangeType, TempSelection.Selection.Array(), SelectionDelta);
+		}
+	});
+}
+
+bool FSkeletalMeshEditingCacheSelector::ConvertSelection(
+	const FGeometrySelection& FromSelectionIn,
+	FGeometrySelection& ToSelectionOut,
+	const EEnumerateSelectionConversionParams ConversionParams)
+{
+	if (!ActiveSubmeshMapping)
+	{
+		return FBaseDynamicMeshSelector::ConvertSelection(FromSelectionIn, ToSelectionOut, ConversionParams);
+	}
+
+	bool bResult = false;
+	EditingMesh->ProcessMesh([&](const FDynamicMesh3& Mesh)
+	{
+		bResult = UE::Geometry::ConvertSelection(Mesh, GetEditingMeshGroupTopology(), FromSelectionIn, ToSelectionOut, ConversionParams);
+	});
+	return bResult;
+}
+
+void FSkeletalMeshEditingCacheSelector::UpdateSoftSelection(
+	FGeometrySelection& Selection,
+	const FUpdateSoftSelectionConfig& InSoftSelectionConfig)
+{
+	if (!ActiveSubmeshMapping)
+	{
+		// No isolation — IDs match, base class runs Dijkstra on preview mesh (= editing mesh topology)
+		FBaseDynamicMeshSelector::UpdateSoftSelection(Selection, InSoftSelectionConfig);
+		return;
+	}
+
+	// With isolation: remap selection to preview space, run Dijkstra on preview mesh, remap results back
+	FGeometrySelection PreviewSelection = RemapSelectionToPreviewMesh(Selection, ActiveSubmeshMapping, FBaseDynamicMeshSelector::GetGroupTopology());
+	FBaseDynamicMeshSelector::UpdateSoftSelection(PreviewSelection, InSoftSelectionConfig);
+
+	// Remap soft-selected vertex IDs from preview space back to editing space
+	Selection.SoftSelectedVertices.Reset();
+	for (const TPair<uint64, double>& Pair : PreviewSelection.SoftSelectedVertices)
+	{
+		int32 PreviewVID = static_cast<int32>(FGeoSelectionID(Pair.Key).GeometryID);
+		int32 EditingVID = RemapVertexToEditingMesh(PreviewVID);
+		if (EditingVID != FDynamicMesh3::InvalidID)
+		{
+			Selection.SoftSelectedVertices.Add(FGeoSelectionID::MeshVertex(EditingVID).Encoded(), Pair.Value);
+		}
+	}
+}
+
+
+// ============================================================================
+// Rendering overrides (remap editing IDs to preview IDs, then use base/preview mesh)
+// ============================================================================
+
+void FSkeletalMeshEditingCacheSelector::AccumulateSelectionElements(
+	const FGeometrySelection& Selection,
+	FGeometrySelectionElements& Elements,
+	bool bTransformToWorld,
+	EEnumerateSelectionMapping Flags)
+{
+	if (!ActiveSubmeshMapping)
+	{
+		FBaseDynamicMeshSelector::AccumulateSelectionElements(Selection, Elements, bTransformToWorld, Flags);
+		return;
+	}
+
+	FGeometrySelection PreviewSelection = RemapSelectionToPreviewMesh(Selection, ActiveSubmeshMapping, FBaseDynamicMeshSelector::GetGroupTopology());
+
+	const FTransform UseWorldTransform = GetLocalToWorldTransform();
+	const FTransform* ApplyTransform = bTransformToWorld ? &UseWorldTransform : nullptr;
+
+	ToolSelectionUtil::AccumulateSelectionElements(
+		Elements,
+		PreviewSelection,
+		TargetMesh->GetMeshRef(),
+		PreviewSelection.TopologyType == EGeometryTopologyType::Polygroup ? FBaseDynamicMeshSelector::GetGroupTopology() : nullptr,
+		ApplyTransform,
+		Flags);
+}
+
+void FSkeletalMeshEditingCacheSelector::AccumulateSelectionBounds(
+	const FGeometrySelection& Selection,
+	FGeometrySelectionBounds& BoundsInOut,
+	bool bTransformToWorld)
+{
+	if (!ActiveSubmeshMapping)
+	{
+		FBaseDynamicMeshSelector::AccumulateSelectionBounds(Selection, BoundsInOut, bTransformToWorld);
+		return;
+	}
+
+	FGeometrySelection PreviewSelection = RemapSelectionToPreviewMesh(Selection, ActiveSubmeshMapping, FBaseDynamicMeshSelector::GetGroupTopology());
+
+	FTransform UseTransform = bTransformToWorld ? GetLocalToWorldTransform() : FTransform::Identity;
+	FAxisAlignedBox3d TargetWorldBounds = FAxisAlignedBox3d::Empty();
+	int32 ElementCount = 0;
+
+	if (PreviewSelection.TopologyType == EGeometryTopologyType::Polygroup)
+	{
+		TargetMesh->ProcessMesh([&](const FDynamicMesh3& SourceMesh)
+		{
+			EnumeratePolygroupSelectionVertices(PreviewSelection, SourceMesh, FBaseDynamicMeshSelector::GetGroupTopology(), UseTransform,
+				[&](uint32 VertexID, const FVector3d& Position) { TargetWorldBounds.Contain(Position); ElementCount++; });
+		});
+	}
+	else
+	{
+		TargetMesh->ProcessMesh([&](const FDynamicMesh3& SourceMesh)
+		{
+			EnumerateTriangleSelectionVertices(PreviewSelection, SourceMesh, &UseTransform,
+				[&](uint32 VertexID, const FVector3d& Position) { TargetWorldBounds.Contain(Position); ElementCount++; });
+		});
+	}
+
+	if (ElementCount > 0)
+	{
+		BoundsInOut.WorldBounds.Contain(TargetWorldBounds);
+	}
+}
+
+void FSkeletalMeshEditingCacheSelector::GetSelectionFrame(
+	const FGeometrySelection& Selection,
+	FFrame3d& SelectionFrame,
+	bool bTransformToWorld)
+{
+	if (!ActiveSubmeshMapping)
+	{
+		FBaseDynamicMeshSelector::GetSelectionFrame(Selection, SelectionFrame, bTransformToWorld);
+		return;
+	}
+
+	FGeometrySelection PreviewSelection = RemapSelectionToPreviewMesh(Selection, ActiveSubmeshMapping, FBaseDynamicMeshSelector::GetGroupTopology());
+
+	if (PreviewSelection.TopologyType == EGeometryTopologyType::Polygroup)
+	{
+		const FGroupTopology* Topology = FBaseDynamicMeshSelector::GetGroupTopology();
+		FGroupTopologySelection TopoSelection;
+		TargetMesh->ProcessMesh([&](const FDynamicMesh3& SourceMesh)
+		{
+			ConvertPolygroupSelectionToTopologySelection(PreviewSelection, SourceMesh, Topology, TopoSelection);
+		});
+		SelectionFrame = Topology->GetSelectionFrame(TopoSelection);
+	}
+	else
+	{
+		TargetMesh->ProcessMesh([&](const FDynamicMesh3& SourceMesh)
+		{
+			GetTriangleSelectionFrame(PreviewSelection, SourceMesh, SelectionFrame);
+		});
+	}
+
+	if (bTransformToWorld)
+	{
+		SelectionFrame.Transform(GetLocalToWorldTransform());
+	}
+}
+
+
+// ============================================================================
+// Raycast overrides (base hits preview mesh, remap output IDs to editing space)
+// ============================================================================
+
+void FSkeletalMeshEditingCacheSelector::UpdateSelectionViaRaycast(
+	const FWorldRayQueryInfo& RayInfo,
+	FGeometrySelectionEditor& SelectionEditor,
+	const FGeometrySelectionUpdateConfig& UpdateConfig,
+	FGeometrySelectionUpdateResult& ResultOut)
+{
+	if (!ActiveSubmeshMapping)
+	{
+		FBaseDynamicMeshSelector::UpdateSelectionViaRaycast(RayInfo, SelectionEditor, UpdateConfig, ResultOut);
+		return;
+	}
+
+	if (IsLockable() && IsLocked())
+	{
+		return;
+	}
+
+	if (SelectionEditor.GetTopologyType() == EGeometryTopologyType::Triangle)
+	{
+		UpdateSelectionViaRaycast_MeshTopology(RayInfo, SelectionEditor, UpdateConfig, ResultOut);
+		return;
+	}
+
+	check(SelectionEditor.GetTopologyType() == EGeometryTopologyType::Polygroup);
+
+	if (SelectionEditor.GetElementType() == EGeometryElementType::Edge)
+	{
+		UpdateSelectionViaRaycast_GroupEdges(RayInfo, SelectionEditor, UpdateConfig, ResultOut);
+	}
+	else
+	{
+		// Raycast into a temporary selection in preview-space, then remap final state and replace real selection.
+		FGeometrySelection TempSelection = RemapSelectionToPreviewMesh(
+			SelectionEditor.GetSelection(), ActiveSubmeshMapping, FBaseDynamicMeshSelector::GetGroupTopology());
+		FGeometrySelectionEditor TempEditor;
+		TempEditor.Initialize(&TempSelection, true);
+
+		FRay3d LocalRay = GetWorldTransformFunc().InverseTransformRay(RayInfo.WorldRay);
+		FGeometrySelectionUpdateResult TempResult;
+		UpdateGroupSelectionViaRaycast(
+			GetColliderMesh(), FBaseDynamicMeshSelector::GetGroupTopology(), &TempEditor,
+			LocalRay, UpdateConfig, TempResult);
+
+		ResultOut.bSelectionModified = TempResult.bSelectionModified;
+		ResultOut.bSelectionMissed = TempResult.bSelectionMissed;
+		if (!TempResult.bSelectionMissed)
+		{
+			TArray<uint64> EditingSpaceIDs = TempSelection.Selection.Array();
+			RemapPreviewIDsToEditingMesh(EditingSpaceIDs,
+				SelectionEditor.GetTopologyType(), SelectionEditor.GetElementType(),
+				ActiveSubmeshMapping, TargetMesh.Get(), EditingMesh.Get(), GetEditingMeshGroupTopology());
+
+			UpdateSelectionWithNewElements(
+				&SelectionEditor, EGeometrySelectionChangeType::Replace,
+				EditingSpaceIDs, &ResultOut.SelectionDelta);
+		}
+	}
+}
+
+void FSkeletalMeshEditingCacheSelector::UpdateSelectionViaRaycast_MeshTopology(
+	const FWorldRayQueryInfo& RayInfo,
+	FGeometrySelectionEditor& SelectionEditor,
+	const FGeometrySelectionUpdateConfig& UpdateConfig,
+	FGeometrySelectionUpdateResult& ResultOut)
+{
+	if (!ActiveSubmeshMapping)
+	{
+		FBaseDynamicMeshSelector::UpdateSelectionViaRaycast_MeshTopology(RayInfo, SelectionEditor, UpdateConfig, ResultOut);
+		return;
+	}
+
+	// Raycast into a temp selection in preview-space, then remap and apply to real selection.
+	// Initialize temp selection with current selection remapped to preview space so removal deltas work.
+	FGeometrySelection TempSelection = RemapSelectionToPreviewMesh(
+		SelectionEditor.GetSelection(), ActiveSubmeshMapping, FBaseDynamicMeshSelector::GetGroupTopology());
+	FGeometrySelectionEditor TempEditor;
+	TempEditor.Initialize(&TempSelection, false);
+
+	FRay3d LocalRay = GetWorldTransformFunc().InverseTransformRay(RayInfo.WorldRay);
+	FGeometrySelectionUpdateResult TempResult;
+	UpdateTriangleSelectionViaRaycast(GetColliderMesh(), &TempEditor, LocalRay, UpdateConfig, TempResult);
+
+	ResultOut.bSelectionModified = TempResult.bSelectionModified;
+	ResultOut.bSelectionMissed = TempResult.bSelectionMissed;
+	if (!TempResult.bSelectionMissed)
+	{
+		TArray<uint64> EditingSpaceIDs = TempSelection.Selection.Array();
+		RemapPreviewIDsToEditingMesh(EditingSpaceIDs,
+			EGeometryTopologyType::Triangle, SelectionEditor.GetElementType(),
+			ActiveSubmeshMapping, TargetMesh.Get(), EditingMesh.Get(), GetEditingMeshGroupTopology());
+
+		UpdateSelectionWithNewElements(
+			&SelectionEditor, EGeometrySelectionChangeType::Replace,
+			EditingSpaceIDs, &ResultOut.SelectionDelta);
+	}
+
+	// Handle dual TriEdgeIDs for edge selection — use editing mesh for topology lookup
+	if (ResultOut.bSelectionModified && SelectionEditor.GetElementType() == EGeometryElementType::Edge)
+	{
+		auto AffectBothTriEdgeIDs = [this, &SelectionEditor, &ResultOut](const EGeometrySelectionChangeType ChangeType)
+		{
+			const TArray<uint64>& DeltaElements = (ChangeType == EGeometrySelectionChangeType::Add ? ResultOut.SelectionDelta.Added : ResultOut.SelectionDelta.Removed);
+			TArray<uint64> ElementsToUpdate;
+			for (const uint64 Element : DeltaElements)
+			{
+				FMeshTriEdgeID TriEdgeID(FGeoSelectionID(Element).GeometryID);
+				EditingMesh->ProcessMesh([TriEdgeID, &ElementsToUpdate](const FDynamicMesh3& SourceMesh)
+				{
+					const int32 EdgeID = SourceMesh.IsTriangle(TriEdgeID.TriangleID) ? SourceMesh.GetTriEdge(TriEdgeID.TriangleID, TriEdgeID.TriEdgeIndex) : IndexConstants::InvalidID;
+					if (SourceMesh.IsEdge(EdgeID))
+					{
+						SourceMesh.EnumerateTriEdgeIDsFromEdgeID(EdgeID,
+							[TriEdgeID, &ElementsToUpdate](const FMeshTriEdgeID OtherTriEdgeID)
+							{
+								if (OtherTriEdgeID.TriangleID != TriEdgeID.TriangleID)
+								{
+									ElementsToUpdate.Add(OtherTriEdgeID.Encoded());
+								}
+							});
+					}
+				});
+			}
+			UpdateSelectionWithNewElements(&SelectionEditor, ChangeType, ElementsToUpdate, &ResultOut.SelectionDelta);
+		};
+		AffectBothTriEdgeIDs(EGeometrySelectionChangeType::Remove);
+		AffectBothTriEdgeIDs(EGeometrySelectionChangeType::Add);
+	}
+}
+
+void FSkeletalMeshEditingCacheSelector::UpdateSelectionViaRaycast_GroupEdges(
+	const FWorldRayQueryInfo& RayInfo,
+	FGeometrySelectionEditor& SelectionEditor,
+	const FGeometrySelectionUpdateConfig& UpdateConfig,
+	FGeometrySelectionUpdateResult& ResultOut)
+{
+	if (!ActiveSubmeshMapping)
+	{
+		FBaseDynamicMeshSelector::UpdateSelectionViaRaycast_GroupEdges(RayInfo, SelectionEditor, UpdateConfig, ResultOut);
+		return;
+	}
+
+	// Raycast into temp selection in preview-space, then remap and apply.
+	// Seed with current selection remapped to preview space so removal deltas work.
+	FGeometrySelection TempSelection = RemapSelectionToPreviewMesh(
+		SelectionEditor.GetSelection(), ActiveSubmeshMapping, FBaseDynamicMeshSelector::GetGroupTopology());
+	FGeometrySelectionEditor TempEditor;
+	TempEditor.Initialize(&TempSelection, true);
+
+	FGeometrySelectionUpdateResult TempResult;
+	FBaseDynamicMeshSelector::UpdateSelectionViaRaycast_GroupEdges(RayInfo, TempEditor, UpdateConfig, TempResult);
+
+	ResultOut.bSelectionModified = TempResult.bSelectionModified;
+	ResultOut.bSelectionMissed = TempResult.bSelectionMissed;
+	if (!TempResult.bSelectionMissed)
+	{
+		TArray<uint64> EditingSpaceIDs = TempSelection.Selection.Array();
+		RemapPreviewIDsToEditingMesh(EditingSpaceIDs,
+			EGeometryTopologyType::Polygroup, EGeometryElementType::Edge,
+			ActiveSubmeshMapping, TargetMesh.Get(), EditingMesh.Get(), GetEditingMeshGroupTopology());
+
+		UpdateSelectionWithNewElements(
+			&SelectionEditor, EGeometrySelectionChangeType::Replace,
+			EditingSpaceIDs, &ResultOut.SelectionDelta);
+	}
+}
+
+void FSkeletalMeshEditingCacheSelector::GetSelectionPreviewForRaycast(
+	const FWorldRayQueryInfo& RayInfo,
+	FGeometrySelectionEditor& PreviewEditor)
+{
+	if (IsLockable() && IsLocked())
+	{
+		return;
+	}
+
+	FGeometrySelectionUpdateResult UpdateResult;
+	UpdateSelectionViaRaycast(RayInfo, PreviewEditor, FGeometrySelectionUpdateConfig(), UpdateResult);
+}
+
+void FSkeletalMeshEditingCacheSelector::UpdateSelectionViaShape(
+	const FWorldShapeQueryInfo& ShapeInfo,
+	FGeometrySelectionEditor& SelectionEditor,
+	const FGeometrySelectionUpdateConfig& UpdateConfig,
+	FGeometrySelectionUpdateResult& ResultOut)
+{
+	if (!ActiveSubmeshMapping)
+	{
+		FBaseDynamicMeshSelector::UpdateSelectionViaShape(ShapeInfo, SelectionEditor, UpdateConfig, ResultOut);
+		return;
+	}
+
+	// Shape-select into temp selection in preview-space, then remap and apply.
+	// Seed with current selection remapped to preview space so removal deltas work.
+	FGeometrySelection TempSelection = RemapSelectionToPreviewMesh(
+		SelectionEditor.GetSelection(), ActiveSubmeshMapping, FBaseDynamicMeshSelector::GetGroupTopology());
+	FGeometrySelectionEditor TempEditor;
+	TempEditor.Initialize(&TempSelection, SelectionEditor.GetTopologyType() == EGeometryTopologyType::Polygroup);
+
+	FGeometrySelectionUpdateResult TempResult;
+	FBaseDynamicMeshSelector::UpdateSelectionViaShape(ShapeInfo, TempEditor, UpdateConfig, TempResult);
+
+	ResultOut.bSelectionModified = TempResult.bSelectionModified;
+	ResultOut.bSelectionMissed = TempResult.bSelectionMissed;
+	if (!TempResult.bSelectionMissed)
+	{
+		TArray<uint64> EditingSpaceIDs = TempSelection.Selection.Array();
+		RemapPreviewIDsToEditingMesh(EditingSpaceIDs,
+			SelectionEditor.GetTopologyType(), SelectionEditor.GetElementType(),
+			ActiveSubmeshMapping, TargetMesh.Get(), EditingMesh.Get(), GetEditingMeshGroupTopology());
+
+		UpdateSelectionWithNewElements(
+			&SelectionEditor, EGeometrySelectionChangeType::Replace,
+			EditingSpaceIDs, &ResultOut.SelectionDelta);
+	}
+}
+
+
+// ============================================================================
+// ID Remapping Helpers
+// ============================================================================
+
+int32 FSkeletalMeshEditingCacheSelector::RemapVertexToEditingMesh(int32 PreviewVertexID) const
+{
+	return ActiveSubmeshMapping ? ActiveSubmeshMapping->MapVertexToBaseMesh(PreviewVertexID) : PreviewVertexID;
+}
+
+int32 FSkeletalMeshEditingCacheSelector::RemapTriangleToEditingMesh(int32 PreviewTriangleID) const
+{
+	return ActiveSubmeshMapping ? ActiveSubmeshMapping->MapTriangleToBaseMesh(PreviewTriangleID) : PreviewTriangleID;
+}
+
+int32 FSkeletalMeshEditingCacheSelector::RemapEdgeToEditingMesh(int32 PreviewEdgeID) const
+{
+	return ActiveSubmeshMapping ? ActiveSubmeshMapping->MapEdgeToBaseMesh(PreviewEdgeID) : PreviewEdgeID;
+}
+
+int32 FSkeletalMeshEditingCacheSelector::RemapVertexToPreviewMesh(int32 EditingVertexID) const
+{
+	return ActiveSubmeshMapping ? ActiveSubmeshMapping->MapVertexToSubmesh(EditingVertexID) : EditingVertexID;
+}
+
+int32 FSkeletalMeshEditingCacheSelector::RemapTriangleToPreviewMesh(int32 EditingTriangleID) const
+{
+	return ActiveSubmeshMapping ? ActiveSubmeshMapping->MapTriangleToSubmesh(EditingTriangleID) : EditingTriangleID;
+}
+
+
+// ============================================================================
+// Isolation Change Handling
+// ============================================================================
+
+void FSkeletalMeshEditingCacheSelector::HandleIsolationChanged()
+{
+	USkeletalMeshEditingCache* Cache = WeakEditingCache.Get();
+	ActiveSubmeshMapping = Cache ? Cache->GetIsolationSubmesh() : nullptr;
+	// Base class spatial caches and NotifyGeometryModified() are already handled by
+	// TargetMesh->OnMeshChanged, which fires when RebuildPreviewMesh replaces the preview mesh content.
+}
+
+
+// ============================================================================
+// Transformation
+// ============================================================================
+
+IGeometrySelectionTransformer* FSkeletalMeshEditingCacheSelector::InitializeTransformation(const FGeometrySelection& Selection)
+{
+	check(!ActiveTransformer);
+
+	if (Selection.TopologyType == EGeometryTopologyType::Polygroup)
+	{
+		auto Transformer = MakeShared<FSkeletalMeshEditingCachePolygroupTransformer>();
+		Transformer->IsolationSubmeshMapping = ActiveSubmeshMapping;
+		Transformer->PreviewGroupTopology = FBaseDynamicMeshSelector::GetGroupTopology();
+		ActiveTransformer = Transformer;
+	}
+	else
+	{
+		auto Transformer = MakeShared<FSkeletalMeshEditingCacheSelectionTransformer>();
+		Transformer->IsolationSubmeshMapping = ActiveSubmeshMapping;
+		Transformer->PreviewGroupTopology = FBaseDynamicMeshSelector::GetGroupTopology();
+		ActiveTransformer = Transformer;
+	}
+
+	ActiveTransformer->bEnableSelectionTransformDrawing = false;
+	ActiveTransformer->Initialize(this);
+	ActiveTransformer->OnEndTransformFunc = [this](IToolsContextTransactionsAPI*)
+	{
+		CommitMeshTransform();
+	};
+	return ActiveTransformer.Get();
+}
+
+void FSkeletalMeshEditingCacheSelector::ShutdownTransformation(IGeometrySelectionTransformer* Transformer)
+{
+	ActiveTransformer.Reset();
+}
+
+void FSkeletalMeshEditingCacheSelector::CommitMeshTransform()
+{
+	USkeletalMeshEditingCache* Cache = WeakEditingCache.Get();
+	if (!Cache)
+	{
+		return;
+	}
+
+	FName MorphTargetName = GetEditingMorphTargetFunc ? GetEditingMorphTargetFunc() : NAME_None;
+	// TargetMesh is the preview mesh — pass it to HandleGeometryUpdate which unposes changes back to editing mesh
+	TargetMesh->ProcessMesh([&](const FDynamicMesh3& PreviewMesh)
+	{
+		Cache->HandleGeometryUpdate(PreviewMesh, MorphTargetName);
+	});
+}
+
+
+// ============================================================================
+// Factory
+// ============================================================================
+
+void FSkeletalMeshEditingCacheSelectorFactory::Init(USkeletalMeshModelingToolsEditorMode* InEditorMode)
+{
+	WeakEditorMode = InEditorMode;
+}
+
+bool FSkeletalMeshEditingCacheSelectorFactory::CanBuildForTarget(FGeometryIdentifier TargetIdentifier) const
+{
+	if (USkeletalMeshModelingToolsEditorMode* EditorMode = WeakEditorMode.Get())
+	{
+		if (EditorMode->GetCurrentEditingCache())
+		{
+			return true;
+		}
+	}
+	return false;
+}
+
+TUniquePtr<IGeometrySelector> FSkeletalMeshEditingCacheSelectorFactory::BuildForTarget(FGeometryIdentifier TargetIdentifier) const
+{
+	USkeletalMeshModelingToolsEditorMode* EditorMode = WeakEditorMode.Get();
+	check(EditorMode);
+	USkeletalMeshEditingCache* Cache = EditorMode->GetCurrentEditingCache();
+	check(Cache);
+
+	TUniquePtr<FSkeletalMeshEditingCacheSelector> Selector = MakeUnique<FSkeletalMeshEditingCacheSelector>();
+
+	TWeakObjectPtr<USkeletalMeshModelingToolsEditorMode> WeakMode = EditorMode;
+	Selector->InitializeFromEditingCache(
+		TargetIdentifier,
+		Cache,
+		[WeakMode]() { return WeakMode.IsValid() ? WeakMode->GetEditingMorphTarget() : NAME_None; }
+	);
+
+	return Selector;
+}
diff --git a/Engine/Plugins/Animation/SkeletalMeshModelingTools/Source/SkeletalMeshModelingTools/Private/Selection/SKMEditingCacheSelector.h b/Engine/Plugins/Animation/SkeletalMeshModelingTools/Source/SkeletalMeshModelingTools/Private/Selection/SKMEditingCacheSelector.h
new file mode 100644
--- /dev/null
+++ b/Engine/Plugins/Animation/SkeletalMeshModelingTools/Source/SkeletalMeshModelingTools/Private/Selection/SKMEditingCacheSelector.h
@@ -0,0 +1,211 @@
+// Copyright Epic Games, Inc. All Rights Reserved.
+
+#pragma once
+
+#include "CoreMinimal.h"
+#include "Selection/DynamicMeshSelector.h"
+#include "Selection/DynamicMeshPolygroupTransformer.h"
+#include "Templates/PimplPtr.h"
+
+#define UE_API SKELETALMESHMODELINGTOOLS_API
+
+namespace UE::Geometry
+{
+	struct FDynamicSubmesh3;
+	class FGroupTopology;
+}
+
+class USkeletalMeshEditingCache;
+class USkeletalMeshModelingToolsEditorMode;
+
+/**
+ * Custom transformers that remap selection from editing-mesh space to preview-mesh space
+ * before delegating to the base class. Only BeginTransform is overridden.
+ */
+class FSkeletalMeshEditingCacheSelectionTransformer : public FBasicDynamicMeshSelectionTransformer
+{
+public:
+	const UE::Geometry::FDynamicSubmesh3* IsolationSubmeshMapping = nullptr;
+	const UE::Geometry::FGroupTopology* PreviewGroupTopology = nullptr;
+	UE_API virtual void BeginTransform(const UE::Geometry::FGeometrySelection& Selection) override;
+};
+
+class FSkeletalMeshEditingCachePolygroupTransformer : public FDynamicMeshPolygroupTransformer
+{
+public:
+	const UE::Geometry::FDynamicSubmesh3* IsolationSubmeshMapping = nullptr;
+	const UE::Geometry::FGroupTopology* PreviewGroupTopology = nullptr;
+	UE_API virtual void BeginTransform(const UE::Geometry::FGeometrySelection& Selection) override;
+};
+
+
+/**
+ * Dual-mesh selector for skeletal mesh editing.
+ *
+ * Base-class TargetMesh is the posed preview mesh (from UPreviewMesh's DynamicMeshComponent).
+ * Spatial queries, hit testing, transformers, and rendering all use the preview mesh naturally.
+ *
+ * Selection IDs are always stored in editing-mesh (full mesh) space. When geometry isolation is
+ * active, the preview mesh is a submesh with different IDs, so the selector remaps between the
+ * two spaces via FDynamicSubmesh3 mappings. When no isolation is active, IDs are identical and
+ * no remapping is needed.
+ *
+ * A separate EditingMesh reference provides access to the full unposed mesh for topology queries
+ * and selection ID construction.
+ */
+class FSkeletalMeshEditingCacheSelector : public FBaseDynamicMeshSelector
+{
+public:
+	UE_API virtual ~FSkeletalMeshEditingCacheSelector();
+
+	/** Initialize from an editing cache. TargetMesh = preview mesh, stores editing mesh separately. */
+	UE_API void InitializeFromEditingCache(
+		FGeometryIdentifier InSourceGeometryIdentifier,
+		USkeletalMeshEditingCache* InEditingCache,
+		TFunction<FName()> InGetEditingMorphTargetFunc);
+
+	/** Set the active submesh mapping for ID remapping. nullptr when isolation is not active. */
+	UE_API void SetActiveSubmeshMapping(const UE::Geometry::FDynamicSubmesh3* InSubmeshMapping);
+
+	// --- Lifecycle ---
+
+	UE_API virtual void Shutdown() override;
+
+	// --- Selection topology overrides (use EditingMesh) ---
+
+	UE_API virtual void InitializeSelectionFromPredicate(
+		FGeometrySelection& SelectionInOut,
+		TFunctionRef<bool(UE::Geometry::FGeoSelectionID)> SelectionIDPredicate,
+		EInitializeSelectionMode InitializeMode = EInitializeSelectionMode::All,
+		const FGeometrySelection* ReferenceSelection = nullptr) override;
+
+	UE_API virtual void UpdateSelectionFromSelection(
+		const FGeometrySelection& FromSelection,
+		bool bAllowConversion,
+		FGeometrySelectionEditor& SelectionEditor,
+		const FGeometrySelectionUpdateConfig& UpdateConfig,
+		UE::Geometry::FGeometrySelectionDelta* SelectionDelta = nullptr) override;
+
+	UE_API virtual bool ConvertSelection(
+		const FGeometrySelection& FromSelectionIn,
+		FGeometrySelection& ToSelectionOut,
+		const UE::Geometry::EEnumerateSelectionConversionParams ConversionParams) override;
+
+	UE_API virtual void UpdateSoftSelection(
+		FGeometrySelection& Selection,
+		const FUpdateSoftSelectionConfig& InSoftSelectionConfig) override;
+
+	// --- Rendering overrides (remap editing IDs to preview IDs for visualization) ---
+
+	UE_API virtual void AccumulateSelectionElements(
+		const FGeometrySelection& Selection,
+		FGeometrySelectionElements& Elements,
+		bool bTransformToWorld,
+		UE::Geometry::EEnumerateSelectionMapping Flags = UE::Geometry::EEnumerateSelectionMapping::Default) override;
+
+	UE_API virtual void AccumulateSelectionBounds(
+		const FGeometrySelection& Selection,
+		FGeometrySelectionBounds& BoundsInOut,
+		bool bTransformToWorld) override;
+
+	UE_API virtual void GetSelectionFrame(
+		const FGeometrySelection& Selection,
+		UE::Geometry::FFrame3d& SelectionFrame,
+		bool bTransformToWorld) override;
+
+	// --- Raycast overrides (hit preview mesh, emit editing-mesh IDs) ---
+
+	UE_API virtual void UpdateSelectionViaRaycast(
+		const FWorldRayQueryInfo& RayInfo,
+		FGeometrySelectionEditor& SelectionEditor,
+		const FGeometrySelectionUpdateConfig& UpdateConfig,
+		FGeometrySelectionUpdateResult& ResultOut) override;
+
+	UE_API virtual void GetSelectionPreviewForRaycast(
+		const FWorldRayQueryInfo& RayInfo,
+		FGeometrySelectionEditor& PreviewEditor) override;
+
+	UE_API virtual void UpdateSelectionViaShape(
+		const FWorldShapeQueryInfo& ShapeInfo,
+		FGeometrySelectionEditor& SelectionEditor,
+		const FGeometrySelectionUpdateConfig& UpdateConfig,
+		FGeometrySelectionUpdateResult& ResultOut) override;
+
+	// --- Transformation ---
+
+	UE_API virtual IGeometrySelectionTransformer* InitializeTransformation(const FGeometrySelection& Selection) override;
+	UE_API virtual void ShutdownTransformation(IGeometrySelectionTransformer* Transformer) override;
+
+protected:
+	// --- Raycast sub-method overrides ---
+
+	UE_API virtual void UpdateSelectionViaRaycast_MeshTopology(
+		const FWorldRayQueryInfo& RayInfo,
+		FGeometrySelectionEditor& SelectionEditor,
+		const FGeometrySelectionUpdateConfig& UpdateConfig,
+		FGeometrySelectionUpdateResult& ResultOut) override;
+
+	UE_API virtual void UpdateSelectionViaRaycast_GroupEdges(
+		const FWorldRayQueryInfo& RayInfo,
+		FGeometrySelectionEditor& SelectionEditor,
+		const FGeometrySelectionUpdateConfig& UpdateConfig,
+		FGeometrySelectionUpdateResult& ResultOut) override;
+
+	// --- Editing Mesh State ---
+
+	TWeakObjectPtr<UDynamicMesh> EditingMesh;
+	FDelegateHandle EditingMesh_OnMeshChangedHandle;
+	FGeometryIdentifier EditingMeshIdentifier;
+	TPimplPtr<UE::Geometry::FGroupTopology> EditingGroupTopology;
+
+	UE_API const UE::Geometry::FGroupTopology* GetEditingMeshGroupTopology();
+	UE_API void InvalidateEditingMeshCaches();
+
+	// --- Submesh Mapping ---
+
+	/** Non-owning pointer to the submesh mapping. nullptr when isolation is not active. */
+	const UE::Geometry::FDynamicSubmesh3* ActiveSubmeshMapping = nullptr;
+
+	// --- ID Remapping Helpers ---
+
+	/** Remap from preview-mesh space to editing-mesh space. Identity when no submesh mapping active. */
+	int32 RemapVertexToEditingMesh(int32 PreviewVertexID) const;
+	int32 RemapTriangleToEditingMesh(int32 PreviewTriangleID) const;
+	int32 RemapEdgeToEditingMesh(int32 PreviewEdgeID) const;
+
+	/** Remap from editing-mesh space to preview-mesh space. Returns InvalidID if not in submesh. */
+	int32 RemapVertexToPreviewMesh(int32 EditingVertexID) const;
+	int32 RemapTriangleToPreviewMesh(int32 EditingTriangleID) const;
+
+	// --- Isolation Change Handling ---
+
+	FDelegateHandle OnIsolationChangedHandle;
+	UE_API void HandleIsolationChanged();
+
+	// --- Transformation State ---
+
+	TWeakObjectPtr<USkeletalMeshEditingCache> WeakEditingCache;
+	TFunction<FName()> GetEditingMorphTargetFunc;
+	TSharedPtr<FBasicDynamicMeshSelectionTransformer> ActiveTransformer;
+	UE_API void CommitMeshTransform();
+
+	friend class FSkeletalMeshEditingCacheSelectorFactory;
+};
+
+
+/**
+ * Factory that constructs FSkeletalMeshEditingCacheSelector instances.
+ */
+class FSkeletalMeshEditingCacheSelectorFactory : public IGeometrySelectorFactory
+{
+public:
+	UE_API void Init(USkeletalMeshModelingToolsEditorMode* InEditorMode);
+
+	UE_API virtual bool CanBuildForTarget(FGeometryIdentifier TargetIdentifier) const override;
+	UE_API virtual TUniquePtr<IGeometrySelector> BuildForTarget(FGeometryIdentifier TargetIdentifier) const override;
+
+protected:
+	TWeakObjectPtr<USkeletalMeshModelingToolsEditorMode> WeakEditorMode;
+};
+
+#undef UE_API
diff --git a/Engine/Plugins/Animation/SkeletalMeshModelingTools/Source/SkeletalMeshModelingTools/Private/Selection/SKMPreviewDynaMeshSelector.cpp b/Engine/Plugins/Animation/SkeletalMeshModelingTools/Source/SkeletalMeshModelingTools/Private/Selection/SKMPreviewDynaMeshSelector.cpp
deleted file mode 100644
diff --git a/Engine/Plugins/Animation/SkeletalMeshModelingTools/Source/SkeletalMeshModelingTools/Private/Selection/SKMPreviewDynaMeshSelector.h b/Engine/Plugins/Animation/SkeletalMeshModelingTools/Source/SkeletalMeshModelingTools/Private/Selection/SKMPreviewDynaMeshSelector.h
deleted file mode 100644
diff --git a/Engine/Plugins/Animation/SkeletalMeshModelingTools/Source/SkeletalMeshModelingTools/Private/SkeletalMeshEditingCache.cpp b/Engine/Plugins/Animation/SkeletalMeshModelingTools/Source/SkeletalMeshModelingTools/Private/SkeletalMeshEditingCache.cpp
--- a/Engine/Plugins/Animation/SkeletalMeshModelingTools/Source/SkeletalMeshModelingTools/Private/SkeletalMeshEditingCache.cpp
+++ b/Engine/Plugins/Animation/SkeletalMeshModelingTools/Source/SkeletalMeshModelingTools/Private/SkeletalMeshEditingCache.cpp
@@ -17,6 +17,9 @@
 #include "EditorViewportClient.h"
 #include "SceneView.h"
 #include "Engine/Engine.h"
+#include "Selections/GeometrySelectionUtil.h"
+#include "InteractiveToolChange.h"
+#include "ToolContextInterfaces.h"
 
 
 #include UE_INLINE_GENERATED_CPP_BY_NAME(SkeletalMeshEditingCache)
@@ -51,11 +54,12 @@ void FSkeletalMeshEditingCacheNotifier::HandleNotification(const TArray<FName>&
 	Notify(BoneNames, InNotifyType);
 }
 
-void USkeletalMeshEditingCache::Spawn(UWorld* World, USkeletalMeshComponent* InSkeletalMeshComponent, EMeshLODIdentifier InLOD, const FDelegates& InDelegates)
+void USkeletalMeshEditingCache::Spawn(UWorld* World, USkeletalMeshComponent* InSkeletalMeshComponent, EMeshLODIdentifier InLOD, const FDelegates& InDelegates, IToolsContextTransactionsAPI* InTransactionsAPI)
 {
 	SetFlags(RF_Transactional);
 	
 	Delegates = InDelegates;
+	TransactionsAPI = InTransactionsAPI;
 	SkeletalMeshComponent = InSkeletalMeshComponent;
 	
 	HostActor = World->SpawnActor(AActor::StaticClass());
@@ -122,18 +126,9 @@ void USkeletalMeshEditingCache::Spawn(UWorld* World, USkeletalMeshComponent* InS
 			{
 				return GetEditingMeshComponent()->GetChangeCount();
 			},
-		[this](int32 ChangeCount)
+		[this](int32)
 			{
-				
-				UDynamicMesh* DynamicMesh = EditingMeshComponent->GetDynamicMesh();
-				FDynamicMesh3 PreviewMeshObject;
-				DynamicMesh->ProcessMesh([&](const FDynamicMesh3& ReadMesh)
-					{
-						PreviewMeshObject.Copy(ReadMesh);
-					});
-
-				PreviewMesh->ReplaceMesh(MoveTemp(PreviewMeshObject));	
-				DeformPreviewMesh(GetComponentSpaceBoneTransforms(), GetMorphTargetWeights());
+				RebuildPreviewMesh();
 			},
 		GetEditingMeshComponent()->GetChangeCount());
 
@@ -765,7 +760,7 @@ void USkeletalMeshEditingCache::HandleGeometryUpdate(const FDynamicMesh3& InMesh
 
 	TMap<FName, float> MorphTargetWeights = GetMorphTargetWeights();
 
-	EditingMeshComponent->HandleGeometryUpdate(MorphTargetName, InMesh, BoneMatrices, MorphTargetWeights);
+	EditingMeshComponent->HandleGeometryUpdate(MorphTargetName, InMesh, BoneMatrices, MorphTargetWeights, GetIsolationSubmesh());
 }
 
 URefSkeletonPoser* USkeletalMeshEditingCache::GetSkeletonPoser() const
@@ -808,15 +803,32 @@ void USkeletalMeshEditingCache::DeformPreviewMesh(const TArray<FTransform>& Comp
 	TArray<FMatrix> BoneMatrices = ComputeBoneMatrices(GetComponentSpaceBoneTransformsRefPose(), ComponentSpaceTransforms);
 			
 	FDynamicMesh3& EditingMesh = *EditingMeshComponent->GetMesh();
-				
+
+	// When isolated, only pose the vertices in the submesh and remap IDs to preview mesh space
+	const FDynamicSubmesh3* Submesh = GetIsolationSubmesh();
+	TArray<int32> IsolatedVertArray;
+	if (Submesh)
+	{
+		const int32 NumSubmeshVerts = Submesh->GetSubmesh().VertexCount();
+		IsolatedVertArray.Reserve(NumSubmeshVerts);
+		for (int32 SubVID : Submesh->GetSubmesh().VertexIndicesItr())
+		{
+			IsolatedVertArray.Add(Submesh->MapVertexToBaseMesh(SubVID));
+		}
+	}
+
 	auto DeformPreviewMeshFunc = [&](FDynamicMesh3& VisualMesh)
 		{
 			auto WriteFunc = [&](FVertInfo VertInfo, const FVector& PosedVertPos)
 				{
-					VisualMesh.SetVertex(VertInfo.VertID, PosedVertPos);
+					int32 PreviewVID = Submesh ? Submesh->MapVertexToSubmesh(VertInfo.VertID) : VertInfo.VertID;
+					if (PreviewVID != FDynamicMesh3::InvalidID)
+					{
+						VisualMesh.SetVertex(PreviewVID, PosedVertPos);
+					}
 				};
 					
-			GetPosedMesh(WriteFunc, EditingMesh, BoneMatrices, NAME_None, MorphTargetWeights);
+			GetPosedMesh(WriteFunc, EditingMesh, BoneMatrices, NAME_None, MorphTargetWeights, IsolatedVertArray);
 			FMeshNormals::QuickRecomputeOverlayNormals(VisualMesh);
 		};
 		
@@ -880,4 +892,265 @@ UDebugSkelMeshComponent* USkeletalMeshEditingCache::GetDebugSkelMeshComponent()
 }
 
 
+// --- Geometry Isolation ---
+
+const TArray<int32>& USkeletalMeshEditingCache::GetIsolatedTriangles() const
+{
+	return IsolatedTrianglesFromFullMesh;
+}
+
+bool USkeletalMeshEditingCache::HasIsolation() const
+{
+	return !IsolatedTrianglesFromFullMesh.IsEmpty();
+}
+
+const UE::Geometry::FDynamicSubmesh3* USkeletalMeshEditingCache::GetIsolationSubmesh() const
+{
+	return IsolationSubmesh.IsSet() ? &IsolationSubmesh.GetValue() : nullptr;
+}
+
+// --- Pure isolation computation ---
+
+bool USkeletalMeshEditingCache::ComputeIsolatedTriangles(
+	EIsolationOperation Operation,
+	const UE::Geometry::FGeometrySelection& Selection,
+	TArray<int32>& OutFullMeshTriangles) const
+{
+	using namespace UE::Geometry;
+
+	// Selection IDs are in editing-mesh (full mesh) space — enumerate against the editing mesh
+	TSet<int32> SelectedTriIDs;
+	GetEditingMeshComponent()->ProcessMesh([&](const FDynamicMesh3& EditingMesh)
+	{
+		EnumerateSelectionTriangles(Selection, EditingMesh,
+			[&SelectedTriIDs](int32 TriangleID)
+			{
+				SelectedTriIDs.Add(TriangleID);
+			});
+	});
+
+	if (SelectedTriIDs.IsEmpty())
+	{
+		return false;
+	}
+
+	// Triangle IDs are already in full-mesh space — no remapping needed
+	OutFullMeshTriangles.Reset();
+
+	if (Operation == EIsolationOperation::Isolate)
+	{
+		OutFullMeshTriangles = SelectedTriIDs.Array();
+	}
+	else // Hide: all currently visible triangles that are NOT selected
+	{
+		if (IsolationSubmesh.IsSet())
+		{
+			// Currently isolated — hide from the isolated set
+			for (int32 TriID : IsolatedTrianglesFromFullMesh)
+			{
+				if (!SelectedTriIDs.Contains(TriID))
+				{
+					OutFullMeshTriangles.Add(TriID);
+				}
+			}
+		}
+		else
+		{
+			// Full mesh visible — hide from all triangles
+			GetEditingMeshComponent()->ProcessMesh([&](const FDynamicMesh3& EditingMesh)
+			{
+				for (int32 TriID : EditingMesh.TriangleIndicesItr())
+				{
+					if (!SelectedTriIDs.Contains(TriID))
+					{
+						OutFullMeshTriangles.Add(TriID);
+					}
+				}
+			});
+		}
+	}
+
+	return true;
+}
+
+// --- High-level isolation APIs ---
+
+bool USkeletalMeshEditingCache::IsolateSelection(const UE::Geometry::FGeometrySelection& InSelection)
+{
+	TArray<int32> NewTriangles;
+	if (!ComputeIsolatedTriangles(EIsolationOperation::Isolate, InSelection, NewTriangles))
+	{
+		return false;
+	}
+
+	SetTriangleIsolation(MoveTemp(NewTriangles));
+	return true;
+}
+
+bool USkeletalMeshEditingCache::HideSelection(const UE::Geometry::FGeometrySelection& InSelection)
+{
+	TArray<int32> NewTriangles;
+	if (!ComputeIsolatedTriangles(EIsolationOperation::Hide, InSelection, NewTriangles))
+	{
+		return false;
+	}
+
+	SetTriangleIsolation(MoveTemp(NewTriangles));
+
+	return true;
+}
+
+bool USkeletalMeshEditingCache::ShowFullMesh(bool bSaveIsolation)
+{
+	if (!HasIsolation())
+	{
+		return false;
+	}
+
+	if (bSaveIsolation)
+	{
+		SavedIsolationTriangles = IsolatedTrianglesFromFullMesh;
+	}
+
+	SetTriangleIsolation({});
+	return true;
+}
+
+bool USkeletalMeshEditingCache::HasSavedIsolation() const
+{
+	return SavedIsolationTriangles.IsSet();
+}
+
+void USkeletalMeshEditingCache::RestoreSavedIsolation()
+{
+	if (SavedIsolationTriangles.IsSet())
+	{
+		SetTriangleIsolation(MoveTemp(SavedIsolationTriangles.GetValue()));
+		SavedIsolationTriangles.Reset();
+	}
+}
+
+void USkeletalMeshEditingCache::DiscardSavedIsolation()
+{
+	SavedIsolationTriangles.Reset();
+}
+
+bool USkeletalMeshEditingCache::ConvertIsolationForTopologyMode(UE::Geometry::EGeometryTopologyType NewTopologyType)
+{
+	using namespace UE::Geometry;
+
+	if (!HasIsolation())
+	{
+		return false;
+	}
+
+	if (NewTopologyType != EGeometryTopologyType::Polygroup)
+	{
+		// Triangle mode is the finest granularity — no expansion needed
+		return false;
+	}
+
+	const FGroupTopology& FullMeshTopology = GetEditingMeshComponent()->GetGroupTopology();
+
+	// Collect all group IDs touched by currently isolated triangles
+	TSet<int32> TouchedGroupIDs;
+	for (int32 TriID : IsolatedTrianglesFromFullMesh)
+	{
+		TouchedGroupIDs.Add(FullMeshTopology.GetGroupID(TriID));
+	}
+
+	// Expand to all triangles in those groups
+	TArray<int32> ExpandedTriangles;
+	for (int32 GroupID : TouchedGroupIDs)
+	{
+		ExpandedTriangles.Append(FullMeshTopology.GetGroupTriangles(GroupID));
+	}
+
+	if (ExpandedTriangles.Num() == IsolatedTrianglesFromFullMesh.Num())
+	{
+		// Already aligned with polygroup boundaries — nothing to do
+		return false;
+	}
+
+	SetTriangleIsolation(MoveTemp(ExpandedTriangles));
+	return true;
+}
+
+// --- Isolation state mutation ---
+
+namespace
+{
+	class FGeometryIsolationChange : public FToolCommandChange
+	{
+	public:
+		TArray<int32> TrianglesBefore;
+		TArray<int32> TrianglesAfter;
+
+		virtual void Apply(UObject* Object) override
+		{
+			if (USkeletalMeshEditingCache* Cache = Cast<USkeletalMeshEditingCache>(Object))
+			{
+				Cache->SetTriangleIsolation(TrianglesAfter);
+			}
+		}
+
+		virtual void Revert(UObject* Object) override
+		{
+			if (USkeletalMeshEditingCache* Cache = Cast<USkeletalMeshEditingCache>(Object))
+			{
+				Cache->SetTriangleIsolation(TrianglesBefore);
+			}
+		}
+
+		virtual FString ToString() const override
+		{
+			return TEXT("Geometry Isolation Change");
+		}
+	};
+}
+
+void USkeletalMeshEditingCache::SetTriangleIsolation(TArray<int32> InFullMeshTriangles)
+{
+	using namespace UE::Geometry;
+
+	TArray<int32> OldTriangles = MoveTemp(IsolatedTrianglesFromFullMesh);
+
+	// Build new submesh from the incoming triangles
+	IsolatedTrianglesFromFullMesh = MoveTemp(InFullMeshTriangles);
+	IsolationSubmesh.Reset();
+	if (!IsolatedTrianglesFromFullMesh.IsEmpty())
+	{
+		GetEditingMeshComponent()->ProcessMesh([this](const FDynamicMesh3& FullMesh)
+		{
+			IsolationSubmesh = FDynamicSubmesh3(&FullMesh, IsolatedTrianglesFromFullMesh);
+		});
+	}
+
+	RebuildPreviewMesh();
+	OnIsolationChanged.Broadcast();
+
+	// Record undo
+	if (TransactionsAPI && !GIsTransacting)
+	{
+		TUniquePtr<FGeometryIsolationChange> Change = MakeUnique<FGeometryIsolationChange>();
+		Change->TrianglesBefore = MoveTemp(OldTriangles);
+		Change->TrianglesAfter = IsolatedTrianglesFromFullMesh;
+		TransactionsAPI->AppendChange(this, MoveTemp(Change), LOCTEXT("IsolationChange", "Isolation Change"));
+	}
+}
+
+void USkeletalMeshEditingCache::RebuildPreviewMesh()
+{
+	using namespace UE::Geometry;
+
+	GetEditingMeshComponent()->ProcessMesh([&](const FDynamicMesh3& FullMesh)
+	{
+			const FDynamicMesh3* NewPreviewMesh = IsolationSubmesh ? &IsolationSubmesh->GetSubmesh() : &FullMesh;
+			PreviewMesh->ReplaceMesh(*NewPreviewMesh);
+	});
+
+	DeformPreviewMesh(GetComponentSpaceBoneTransforms(), GetMorphTargetWeights());
+}
+
+
 #undef LOCTEXT_NAMESPACE
diff --git a/Engine/Plugins/Animation/SkeletalMeshModelingTools/Source/SkeletalMeshModelingTools/Private/SkeletalMeshEditingCache.h b/Engine/Plugins/Animation/SkeletalMeshModelingTools/Source/SkeletalMeshModelingTools/Private/SkeletalMeshEditingCache.h
--- a/Engine/Plugins/Animation/SkeletalMeshModelingTools/Source/SkeletalMeshModelingTools/Private/SkeletalMeshEditingCache.h
+++ b/Engine/Plugins/Animation/SkeletalMeshModelingTools/Source/SkeletalMeshModelingTools/Private/SkeletalMeshEditingCache.h
@@ -6,6 +6,9 @@
 #include "SkeletalMeshNotifier.h"
 #include "Animation/DebugSkelMeshComponent.h"
 #include "Changes/ValueWatcher.h"
+#include "DynamicSubmesh3.h"
+#include "GroupTopology.h"
+#include "Selections/GeometrySelection.h"
 #include "SkeletalMesh/SkeletalMeshToolsHelper.h"
 #include "TargetInterfaces/MeshTargetInterfaceTypes.h"
 
@@ -24,6 +27,7 @@ class FEditorViewportClient;
 class FViewport;
 class FSceneView;
 class FCanvas;
+class IToolsContextTransactionsAPI;
 
 class USkeletalMeshEditingCache;
 class FSkeletalMeshEditingCacheNotifier: public ISkeletalMeshNotifier
@@ -52,7 +56,7 @@ public:
 		FOnGetSkeletalMeshSkeletonNotifier OnGetSkeletalMeshSkeletonNotifierDelegate;
 	};
 
-	void Spawn(UWorld* World, USkeletalMeshComponent* SkeletalMeshComponent, EMeshLODIdentifier LOD, const FDelegates& InDelegates);
+	void Spawn(UWorld* World, USkeletalMeshComponent* SkeletalMeshComponent, EMeshLODIdentifier LOD, const FDelegates& InDelegates, IToolsContextTransactionsAPI* InTransactionsAPI);
 	EMeshLODIdentifier GetLOD() const;
 
 	void Destroy();
@@ -107,6 +111,32 @@ public:
 	void HandleGeometryUpdate(const UE::Geometry::FDynamicMesh3& InMesh, FName MorphTargetName);
 	
 	URefSkeletonPoser* GetSkeletonPoser() const;
+
+	// Geometry isolation — high-level APIs that compute 
+	// new isolation, mutate state, rebuild preview mesh
+	// Returns false if no-op.
+	bool IsolateSelection(const UE::Geometry::FGeometrySelection& InSelection);
+	bool HideSelection(const UE::Geometry::FGeometrySelection& InSelection);
+	bool ShowFullMesh(bool bSaveIsolation = false);
+	
+	// Converts isolation triangles when switching topology modes (e.g. expands to full polygroups).
+	// Returns false if no-op. 
+	bool ConvertIsolationForTopologyMode(UE::Geometry::EGeometryTopologyType NewTopologyType);
+
+	bool HasIsolation() const;
+	const TArray<int32>& GetIsolatedTriangles() const;
+	const UE::Geometry::FDynamicSubmesh3* GetIsolationSubmesh() const;
+
+	FSimpleMulticastDelegate& GetOnIsolationChanged() { return OnIsolationChanged; }
+
+	// Saved isolation — allows temporarily showing full mesh and restoring later (e.g. during tool use)
+	bool HasSavedIsolation() const;
+	void RestoreSavedIsolation();
+	void DiscardSavedIsolation();
+
+
+	// Sets isolation state and rebuilds the preview mesh. Records undo via TransactionsAPI.
+	void SetTriangleIsolation(TArray<int32> InFullMeshTriangles);
 	
 protected:
 	UPROPERTY()
@@ -122,6 +152,7 @@ protected:
 	TWeakObjectPtr<USkeletalMeshComponent> SkeletalMeshComponent;
 
 	FDelegates Delegates;
+	IToolsContextTransactionsAPI* TransactionsAPI = nullptr;
 
 	// Flags that drives both dyna mesh and skeletal mesh
 	bool bCacheVisibility = true;
@@ -212,6 +243,27 @@ protected:
 
 	UDebugSkelMeshComponent* GetDebugSkelMeshComponent() const;
 
+	// Geometry isolation state
+	TArray<int32> IsolatedTrianglesFromFullMesh;
+	TOptional<UE::Geometry::FDynamicSubmesh3> IsolationSubmesh;
+	TOptional<TArray<int32>> SavedIsolationTriangles;
+	FSimpleMulticastDelegate OnIsolationChanged;
+
+	void RebuildPreviewMesh();
+
+	// Pure computation: given an operation and current selection, computes the new isolation
+	// triangle set (in full mesh space) and remaps the selection to the new submesh space.
+	enum class EIsolationOperation
+	{
+		Isolate, 
+		Hide
+	};
+	
+	bool ComputeIsolatedTriangles(
+		EIsolationOperation Operation,
+		const UE::Geometry::FGeometrySelection& Selection,
+		TArray<int32>& OutFullMeshTriangles) const;
+
 };
 
 
diff --git a/Engine/Plugins/Animation/SkeletalMeshModelingTools/Source/SkeletalMeshModelingTools/Private/SkeletalMeshModelingToolsCommands.cpp b/Engine/Plugins/Animation/SkeletalMeshModelingTools/Source/SkeletalMeshModelingTools/Private/SkeletalMeshModelingToolsCommands.cpp
--- a/Engine/Plugins/Animation/SkeletalMeshModelingTools/Source/SkeletalMeshModelingTools/Private/SkeletalMeshModelingToolsCommands.cpp
+++ b/Engine/Plugins/Animation/SkeletalMeshModelingTools/Source/SkeletalMeshModelingTools/Private/SkeletalMeshModelingToolsCommands.cpp
@@ -22,6 +22,10 @@ void FSkeletalMeshModelingToolsCommands::RegisterCommands()
 
 	UI_COMMAND(FlipEditingMorphTarget, "Flip Editing Morph Target", "Flip the Morph Target currently selected for editing, if geometry selection is active", EUserInterfaceActionType::Button, FInputChord());
 
+	UI_COMMAND(IsolateSelection, "Isolate Selection", "Show only the selected geometry, hiding everything else", EUserInterfaceActionType::Button, FInputChord());
+	UI_COMMAND(HideSelection, "Hide Selected", "Hide the selected geometry from the viewport", EUserInterfaceActionType::Button, FInputChord(EKeys::H));
+	UI_COMMAND(ShowFullMesh, "Show Full Mesh", "Restore visibility of all geometry", EUserInterfaceActionType::Button, FInputChord());
+
 	UI_COMMAND(ToggleShowAllLocalRotationAxes, "Show All Local Rotation Axes", "When enabled shows all local rotation axes, otherwise for the selected bone only", EUserInterfaceActionType::ToggleButton, FInputChord());	
 }
 
diff --git a/Engine/Plugins/Animation/SkeletalMeshModelingTools/Source/SkeletalMeshModelingTools/Private/SkeletalMeshModelingToolsCommands.h b/Engine/Plugins/Animation/SkeletalMeshModelingTools/Source/SkeletalMeshModelingTools/Private/SkeletalMeshModelingToolsCommands.h
--- a/Engine/Plugins/Animation/SkeletalMeshModelingTools/Source/SkeletalMeshModelingTools/Private/SkeletalMeshModelingToolsCommands.h
+++ b/Engine/Plugins/Animation/SkeletalMeshModelingTools/Source/SkeletalMeshModelingTools/Private/SkeletalMeshModelingToolsCommands.h
@@ -37,6 +37,11 @@ public:
 
 	TSharedPtr<FUICommandInfo> FlipEditingMorphTarget;
 
+	// Geometry Isolation commands
+	TSharedPtr<FUICommandInfo> IsolateSelection;
+	TSharedPtr<FUICommandInfo> HideSelection;
+	TSharedPtr<FUICommandInfo> ShowFullMesh;
+
     /** Command to toggle show all local rotation axes */
 	TSharedPtr<FUICommandInfo> ToggleShowAllLocalRotationAxes;
 };
diff --git a/Engine/Plugins/Animation/SkeletalMeshModelingTools/Source/SkeletalMeshModelingTools/Private/SkeletalMeshModelingToolsEditorMode.cpp b/Engine/Plugins/Animation/SkeletalMeshModelingTools/Source/SkeletalMeshModelingTools/Private/SkeletalMeshModelingToolsEditorMode.cpp
--- a/Engine/Plugins/Animation/SkeletalMeshModelingTools/Source/SkeletalMeshModelingTools/Private/SkeletalMeshModelingToolsEditorMode.cpp
+++ b/Engine/Plugins/Animation/SkeletalMeshModelingTools/Source/SkeletalMeshModelingTools/Private/SkeletalMeshModelingToolsEditorMode.cpp
@@ -3,6 +3,7 @@
 #include "SkeletalMeshModelingToolsEditorMode.h"
 
 #include "AttributeEditorTool.h"
+#include "ContextObjectStore.h"
 #include "SkeletalMeshModelingToolsEditorModeToolkit.h"
 #include "SkeletalMeshModelingToolsCommands.h"
 #include "SkeletalMeshGizmoUtils.h"
@@ -80,6 +81,7 @@
 #include "Components/SKMBackedDynaMeshComponent.h"
 #include "SkeletalMesh/RefSkeletonPoser.h"
 
+#include "SkeletalMesh/ISkeletalMeshGeometryIsolationAwareTool.h"
 #include "SkeletalMesh/SkeletalMeshEditingInterface.h"
 #include "SkeletalMesh/SkeletonEditingTool.h"
 #include "SkeletalMesh/SkinWeightsBindingTool.h"
@@ -90,7 +92,7 @@
 #include "ModelingToolsEditorModeSettings.h"
 #include "Selection/GeometrySelectionManager.h"
 #include "Commands/ModifyGeometrySelectionCommand.h"
-#include "Selection/SKMPreviewDynaMeshSelector.h"
+#include "Selection/SKMEditingCacheSelector.h"
 #include "Widgets/Input/SSpinBox.h"
 #include "SkeletalMeshModelingToolsViewportToolbarExtension.h"
 
@@ -796,6 +798,7 @@ void USkeletalMeshModelingToolsEditorMode::Exit()
 	// is involved (although since we are exiting Mode this currently would never matter)
 	if (SelectionManager != nullptr)
 	{
+		GetInteractiveToolsContext()->ContextObjectStore->RemoveContextObject(SelectionManager);
 		SelectionManager->ClearSelection();
 		SelectionManager->Shutdown();		// will clear active targets
 		SelectionManager = nullptr;
@@ -938,10 +941,12 @@ void USkeletalMeshModelingToolsEditorMode::InitializeGeometrySelectionSystems()
 	SelectionManager = NewObject<UGeometrySelectionManager>(GetToolManager());
 	SelectionManager->Initialize(GetInteractiveToolsContext(), GetToolManager()->GetContextTransactionsAPI());
 	{
-		TUniquePtr<FSkeletalMeshPreviewDynamicMeshSelectorFactory> SelectorFactory =  MakeUnique<FSkeletalMeshPreviewDynamicMeshSelectorFactory>();
+		TUniquePtr<FSkeletalMeshEditingCacheSelectorFactory> SelectorFactory = MakeUnique<FSkeletalMeshEditingCacheSelectorFactory>();
 		SelectorFactory->Init(this);
 		SelectionManager->RegisterSelectorFactory(MoveTemp(SelectorFactory));
 	}
+	// Matches ModelingToolsEditorMode, might be giving too much access... probably better wrap it in a interface
+	GetInteractiveToolsContext()->ContextObjectStore->AddContextObject(SelectionManager);
 
 	SelectionTransformTweaker= NewObject<USkeletalMeshGeometrySelectionTransformTweaker>();
 	SelectionTransformTweaker->Setup(SelectionManager);
@@ -1030,6 +1035,73 @@ void USkeletalMeshModelingToolsEditorMode::InitializeGeometrySelectionSystems()
 		}));
 	CommandList->MapAction(ToolManagerCommands.SelectionHitBackFaces, ToggleHitBackFacesAction);
 
+	// Register Geometry Isolation commands
+	CommandList->MapAction(FSkeletalMeshModelingToolsCommands::Get().IsolateSelection, FUIAction(
+		FExecuteAction::CreateLambda([WeakThis]
+		{
+			if (!WeakThis.IsValid() || !WeakThis->GetToolManager() || !WeakThis->GetToolManager()->GetContextTransactionsAPI())
+			{
+				return;
+			}
+			IToolsContextTransactionsAPI* TransactionsAPI = WeakThis->GetToolManager()->GetContextTransactionsAPI();
+			TransactionsAPI->BeginUndoTransaction(LOCTEXT("IsolateSelectionTransaction", "Isolate Selection"));
+			if (!WeakThis->IsolateSelection())
+			{
+				TransactionsAPI->CancelUndoTransaction();
+				return;
+			}
+
+			TransactionsAPI->EndUndoTransaction();
+		}),
+		FCanExecuteAction::CreateLambda([WeakThis]()
+		{
+			return WeakThis.IsValid() && WeakThis->HasActiveGeometrySelection();
+		})));
+
+	CommandList->MapAction(FSkeletalMeshModelingToolsCommands::Get().HideSelection, FUIAction(
+		FExecuteAction::CreateLambda([WeakThis]
+		{
+			if (!WeakThis.IsValid() || !WeakThis->GetToolManager() || !WeakThis->GetToolManager()->GetContextTransactionsAPI())
+			{
+				return;
+			}
+			IToolsContextTransactionsAPI* TransactionsAPI = WeakThis->GetToolManager()->GetContextTransactionsAPI();
+			TransactionsAPI->BeginUndoTransaction(LOCTEXT("HideSelectionTransaction", "Hide Selection"));
+			if (!WeakThis->HideSelection())
+			{
+				TransactionsAPI->CancelUndoTransaction();
+				return;
+			}
+
+			TransactionsAPI->EndUndoTransaction();
+		}),
+		FCanExecuteAction::CreateLambda([WeakThis]()
+		{
+			return WeakThis.IsValid() && WeakThis->HasActiveGeometrySelection();
+		})));
+
+	CommandList->MapAction(FSkeletalMeshModelingToolsCommands::Get().ShowFullMesh, FUIAction(
+		FExecuteAction::CreateLambda([WeakThis]
+		{
+			if (!WeakThis.IsValid() || !WeakThis->GetToolManager() || !WeakThis->GetToolManager()->GetContextTransactionsAPI())
+			{
+				return;
+			}
+			IToolsContextTransactionsAPI* TransactionsAPI = WeakThis->GetToolManager()->GetContextTransactionsAPI();
+			TransactionsAPI->BeginUndoTransaction(LOCTEXT("ShowFullMeshTransaction", "Show Full Mesh"));
+			if (!WeakThis->ShowFullMesh())
+			{
+				TransactionsAPI->CancelUndoTransaction();
+				return;
+			}
+				
+			TransactionsAPI->EndUndoTransaction();
+		}),
+		FCanExecuteAction::CreateLambda([WeakThis]()
+		{
+			return WeakThis.IsValid() && WeakThis->CurrentEditingCache && WeakThis->CurrentEditingCache->HasIsolation();
+		})));
+
 	// Register Soft Selection Toggle
 	FUIAction ToggleSoftSelectionAction = FUIAction(
 		FExecuteAction::CreateLambda([WeakThis]
@@ -1134,6 +1206,13 @@ void USkeletalMeshModelingToolsEditorMode::InitializeGeometrySelectionSystems()
 
 					IToolsContextTransactionsAPI* TransactionsAPI = ToolManager->GetContextTransactionsAPI();
 					TransactionsAPI->BeginUndoTransaction(LOCTEXT("ChangeSelectionMode", "Selection Mode"));
+					if (WeakThis->GetCurrentEditingCache() && WeakThis->GetCurrentEditingCache()->HasIsolation())
+					{
+						UE::Geometry::EGeometryTopologyType NewTopologyType = (TopoMode == UGeometrySelectionManager::EMeshTopologyMode::Polygroup)
+							? UE::Geometry::EGeometryTopologyType::Polygroup
+							: UE::Geometry::EGeometryTopologyType::Triangle;
+						WeakThis->GetCurrentEditingCache()->ConvertIsolationForTopologyMode(NewTopologyType);
+					}
 					WeakThis->SelectionManager->SetMeshSelectionTypeAndMode(ElementMode, TopoMode,  TopoMode != UGeometrySelectionManager::EMeshTopologyMode::None);
 					TransactionsAPI->EndUndoTransaction();
 					
@@ -1415,6 +1494,21 @@ void USkeletalMeshModelingToolsEditorMode::OnToolStarted(UInteractiveToolManager
 void USkeletalMeshModelingToolsEditorMode::OnToolPostBuild( UInteractiveToolManager* InToolManager, EToolSide InSide, UInteractiveTool* InBuiltTool,
 	UInteractiveToolBuilder* InToolBuilder, const FToolBuilderState& ToolState)
 {
+	// Want to clear active selection when a Tool starts, but we have to wait until after it has been
+	// built, so that the Tool has a chance to see the Selection
+	if (GetSelectionManager() && GetSelectionManager()->HasSelection())
+	{
+		ensureMsgf(!GetSelectionManager()->HasSavedSelection(), TEXT("Selection manager should not already have a saved selection before we save-on-clear here in tool setup."));
+		GetSelectionManager()->ClearSelection(/*bSaveSelectionBeforeClear*/ true);
+	}
+			
+	if (GetCurrentEditingCache()->HasIsolation())
+	{
+		ensureMsgf(!GetCurrentEditingCache()->HasSavedIsolation(), TEXT("Editing Cache should not already have a saved isolation before we save-on-clear here in tool setup."));
+		constexpr bool bSaveIsolation = true;
+		GetCurrentEditingCache()->ShowFullMesh(bSaveIsolation);
+	}
+	
 	// Tools can decide for themselves if they want to enable bone manipulation
 	ToggleBoneManipulation(false);
 
@@ -1425,6 +1519,42 @@ void USkeletalMeshModelingToolsEditorMode::OnToolPostBuild( UInteractiveToolMana
 void USkeletalMeshModelingToolsEditorMode::OnToolEndedWithStatus(UInteractiveToolManager* Manager, UInteractiveTool* Tool,
 	EToolShutdownType ShutdownType)
 {
+	bool bRestoreEverything = (ShutdownType == EToolShutdownType::Cancel);
+	bool bCanRestoreIsolation = false;
+			
+	if (ISkeletalMeshGeometryIsolationAwareTool* IsolateAwareTool = Cast<ISkeletalMeshGeometryIsolationAwareTool>(Tool))
+	{
+		bCanRestoreIsolation = bRestoreEverything || IsolateAwareTool->IsInputIsolationValidOnOutput();
+	}
+			
+	if (bCanRestoreIsolation)
+	{
+		GetCurrentEditingCache()->RestoreSavedIsolation();	
+	}
+	else
+	{
+		GetCurrentEditingCache()->DiscardSavedIsolation();	
+	}
+	ensureMsgf(!GetCurrentEditingCache()->HasSavedIsolation(), TEXT("Editing Cache's saved isolation should be cleared on tool end."));
+			
+
+	bool bCanRestoreSelection = false;
+	if (IInteractiveToolManageGeometrySelectionAPI* ManageSelectionTool = Cast<IInteractiveToolManageGeometrySelectionAPI>(Tool))
+	{
+		bCanRestoreSelection = bRestoreEverything || ManageSelectionTool->IsInputSelectionValidOnOutput();
+	}
+			
+	if (bCanRestoreSelection)
+	{
+		GetSelectionManager()->RestoreSavedSelection();
+	}
+	else
+	{
+		GetSelectionManager()->DiscardSavedSelection();
+	}
+			
+	ensureMsgf(!GetSelectionManager()->HasSavedSelection(), TEXT("Selection manager's saved selection should be cleared on tool end."));	
+	
 	FSkeletalMeshModelingToolsActionCommands::UpdateToolCommandBinding(Tool, Toolkit->GetToolkitCommands(), true);
 	if (ExtensionToolCommandsGetter)
 	{
@@ -1866,6 +1996,8 @@ EMeshLODIdentifier USkeletalMeshModelingToolsEditorMode::GetEditingLOD()
 
 void USkeletalMeshModelingToolsEditorMode::RecreateEditingCache(EMeshLODIdentifier InLOD)
 {
+	SelectionManager->ClearActiveTargets();
+	
 	if (CurrentEditingCache)
 	{
 		CurrentEditingCache->Destroy();
@@ -1903,8 +2035,7 @@ void USkeletalMeshModelingToolsEditorMode::RecreateEditingCache(EMeshLODIdentifi
 
 			return TSharedPtr<ISkeletalMeshNotifier>();
 		});
-	CurrentEditingCache->Spawn(PreviewWorld, GetSkelMeshComponent(), InLOD, Delegates);
-
+	CurrentEditingCache->Spawn(PreviewWorld, GetSkelMeshComponent(), InLOD, Delegates, GetToolManager()->GetContextTransactionsAPI());
 
 	using namespace UE::SkeletalMeshEditorUtils;
 	EditingCacheNotifierBindScope.Reset(new FSkeletalMeshNotifierBindScope(GetModeBinding()->GetNotifier(), CurrentEditingCache->GetNotifier()));
@@ -1912,8 +2043,7 @@ void USkeletalMeshModelingToolsEditorMode::RecreateEditingCache(EMeshLODIdentifi
 	CurrentEditingCache->GetNotifier()->HandleNotification({}, ESkeletalMeshNotifyType::HierarchyChanged);
 	CurrentEditingCache->GetNotifier()->HandleNotification(GetSelectedBones(), ESkeletalMeshNotifyType::BonesSelected);
 	
-	SelectionManager->ClearActiveTargets();
-	SelectionManager->AddActiveTarget(FGeometryIdentifier::PrimitiveComponent(CurrentEditingCache->GetPreviewMeshComponent(), FGeometryIdentifier::EObjectType::DynamicMeshComponent));
+	SelectionManager->AddActiveTarget(FGeometryIdentifier::PrimitiveComponent(CurrentEditingCache->GetEditingMeshComponent(), FGeometryIdentifier::EObjectType::DynamicMeshComponent));
 }
 
 
@@ -2007,4 +2137,47 @@ bool USkeletalMeshModelingToolsEditorMode::IsSkeletalMeshBoneManipulationEnabled
 	return true;
 }
 
+
+// --- Geometry Isolation ---
+
+bool USkeletalMeshModelingToolsEditorMode::IsolateSelection()
+{
+	if (!CurrentEditingCache || !SelectionManager || !SelectionManager->HasSelection())
+	{
+		return false;
+	}
+
+	UPrimitiveComponent* EditingComponent = CurrentEditingCache->GetEditingMeshComponent();
+	UE::Geometry::FGeometrySelection Selection;
+	SelectionManager->GetSelectionForComponent(EditingComponent, Selection);
+
+	return CurrentEditingCache->IsolateSelection(Selection);
+}
+
+bool USkeletalMeshModelingToolsEditorMode::HideSelection()
+{
+	if (!CurrentEditingCache || !SelectionManager || !SelectionManager->HasSelection())
+	{
+		return false;
+	}
+
+	UPrimitiveComponent* EditingComponent = CurrentEditingCache->GetEditingMeshComponent();
+	UE::Geometry::FGeometrySelection Selection;
+	SelectionManager->GetSelectionForComponent(EditingComponent, Selection);
+
+	return CurrentEditingCache->HideSelection(Selection);
+}
+
+bool USkeletalMeshModelingToolsEditorMode::ShowFullMesh()
+{
+	if (!CurrentEditingCache || !SelectionManager)
+	{
+		return false;
+	}
+
+	return CurrentEditingCache->ShowFullMesh();
+}
+
+
+
 #undef LOCTEXT_NAMESPACE
diff --git a/Engine/Plugins/Animation/SkeletalMeshModelingTools/Source/SkeletalMeshModelingTools/Private/SkeletalMeshModelingToolsEditorMode.h b/Engine/Plugins/Animation/SkeletalMeshModelingTools/Source/SkeletalMeshModelingTools/Private/SkeletalMeshModelingToolsEditorMode.h
--- a/Engine/Plugins/Animation/SkeletalMeshModelingTools/Source/SkeletalMeshModelingTools/Private/SkeletalMeshModelingToolsEditorMode.h
+++ b/Engine/Plugins/Animation/SkeletalMeshModelingTools/Source/SkeletalMeshModelingTools/Private/SkeletalMeshModelingToolsEditorMode.h
@@ -300,6 +300,12 @@ private:
 	bool IsGeometrySelectionActive() const;
 	bool HasActiveGeometrySelection() const;
 	void ToggleGeometrySelectionViewportInteractions(bool bEnable);
+	
+	// Geometry isolation (return false if no-op)
+	bool IsolateSelection();
+	bool HideSelection();
+	bool ShowFullMesh();
+
 
 	struct FSkeletonStateFactors
 	{
diff --git a/Engine/Plugins/Animation/SkeletalMeshModelingTools/Source/SkeletalMeshModelingTools/Private/SkeletalMeshModelingToolsViewportToolbarExtension.cpp b/Engine/Plugins/Animation/SkeletalMeshModelingTools/Source/SkeletalMeshModelingTools/Private/SkeletalMeshModelingToolsViewportToolbarExtension.cpp
--- a/Engine/Plugins/Animation/SkeletalMeshModelingTools/Source/SkeletalMeshModelingTools/Private/SkeletalMeshModelingToolsViewportToolbarExtension.cpp
+++ b/Engine/Plugins/Animation/SkeletalMeshModelingTools/Source/SkeletalMeshModelingTools/Private/SkeletalMeshModelingToolsViewportToolbarExtension.cpp
@@ -243,7 +243,22 @@ namespace UE::SkeletalMeshModelingTools
 										FModelingToolsManagerCommands::Get().SelectionHitBackFaces, CommandList));
 							}
 							
-							// Section for Selection Filters
+							// Section for Geometry Isolation
+							{
+								FToolMenuSection& IsolationSection =
+								Submenu->FindOrAddSection("Geometry Isolation", LOCTEXT("GeometryIsolationLabel", "Geometry Isolation"));
+								IsolationSection.AddEntry(
+									FToolMenuEntry::InitMenuEntryWithCommandList(
+										FSkeletalMeshModelingToolsCommands::Get().IsolateSelection, CommandList));
+								IsolationSection.AddEntry(
+									FToolMenuEntry::InitMenuEntryWithCommandList(
+										FSkeletalMeshModelingToolsCommands::Get().HideSelection, CommandList));
+								IsolationSection.AddEntry(
+									FToolMenuEntry::InitMenuEntryWithCommandList(
+										FSkeletalMeshModelingToolsCommands::Get().ShowFullMesh, CommandList));
+							}
+
+							// Section for Soft Selection
 							{
 								FToolMenuSection& SoftSelectionSection =
 								Submenu->FindOrAddSection("Soft Selection", LOCTEXT("SoftSelectionLabel", "Soft Selection"));
diff --git a/Engine/Plugins/Experimental/Animation/SkeletalMeshMorphTargetEditingTools/Source/SkeletalMeshMorphTargetEditingTools/Private/MorphTargetMaskTool.cpp b/Engine/Plugins/Experimental/Animation/SkeletalMeshMorphTargetEditingTools/Source/SkeletalMeshMorphTargetEditingTools/Private/MorphTargetMaskTool.cpp
--- a/Engine/Plugins/Experimental/Animation/SkeletalMeshMorphTargetEditingTools/Source/SkeletalMeshMorphTargetEditingTools/Private/MorphTargetMaskTool.cpp
+++ b/Engine/Plugins/Experimental/Animation/SkeletalMeshMorphTargetEditingTools/Source/SkeletalMeshMorphTargetEditingTools/Private/MorphTargetMaskTool.cpp
@@ -679,6 +679,11 @@ void UMorphTargetMaskTool::OnTick(float DeltaTime)
 	}
 }
 
+bool UMorphTargetMaskTool::IsInputIsolationValidOnOutput() const
+{
+	return true;
+}
+
 void UMorphTargetMaskTool::SetupCommonProperties(const TFunction<void(UMorphTargetEditingToolProperties*)>& InSetupFunction)
 {
 	EditorToolProperties = NewObject<UMorphTargetEditingToolProperties>(this);
diff --git a/Engine/Plugins/Experimental/Animation/SkeletalMeshMorphTargetEditingTools/Source/SkeletalMeshMorphTargetEditingTools/Private/MorphTargetVertexSculptTool.cpp b/Engine/Plugins/Experimental/Animation/SkeletalMeshMorphTargetEditingTools/Source/SkeletalMeshMorphTargetEditingTools/Private/MorphTargetVertexSculptTool.cpp
--- a/Engine/Plugins/Experimental/Animation/SkeletalMeshMorphTargetEditingTools/Source/SkeletalMeshMorphTargetEditingTools/Private/MorphTargetVertexSculptTool.cpp
+++ b/Engine/Plugins/Experimental/Animation/SkeletalMeshMorphTargetEditingTools/Source/SkeletalMeshMorphTargetEditingTools/Private/MorphTargetVertexSculptTool.cpp
@@ -427,6 +427,11 @@ void UMorphTargetVertexSculptTool::HandleSkeletalMeshModified(const TArray<FName
 {
 }
 
+bool UMorphTargetVertexSculptTool::IsInputIsolationValidOnOutput() const
+{
+	return true;
+}
+
 void UMorphTargetVertexSculptTool::PoseSculptMesh(const TArray<FTransform>& ComponentSpaceTransforms, const TMap<FName, float>& MorphTargetWeights)
 {
 	using namespace SkeletalMeshToolsHelper;
diff --git a/Engine/Plugins/Experimental/Animation/SkeletalMeshMorphTargetEditingTools/Source/SkeletalMeshMorphTargetEditingTools/Public/MorphTargetMaskTool.h b/Engine/Plugins/Experimental/Animation/SkeletalMeshMorphTargetEditingTools/Source/SkeletalMeshMorphTargetEditingTools/Public/MorphTargetMaskTool.h
--- a/Engine/Plugins/Experimental/Animation/SkeletalMeshMorphTargetEditingTools/Source/SkeletalMeshMorphTargetEditingTools/Public/MorphTargetMaskTool.h
+++ b/Engine/Plugins/Experimental/Animation/SkeletalMeshMorphTargetEditingTools/Source/SkeletalMeshMorphTargetEditingTools/Public/MorphTargetMaskTool.h
@@ -7,6 +7,7 @@
 #include "MeshVertexAttributePaintToolBase.h"
 #include "SkeletalMesh/SkeletalMeshToolsHelper.h"
 #include "SkeletalMesh/SkeletalMeshEditingInterface.h"
+#include "SkeletalMesh/ISkeletalMeshGeometryIsolationAwareTool.h"
 
 #include "MorphTargetMaskTool.generated.h"
 
@@ -67,7 +68,8 @@ public:
 UCLASS()
 class  UMorphTargetMaskTool: 
 	public UMeshVertexAttributePaintToolBase,
-	public IMorphTargetEditingToolInterface
+	public IMorphTargetEditingToolInterface,
+	public ISkeletalMeshGeometryIsolationAwareTool 
 {
 	GENERATED_BODY()
 
@@ -87,6 +89,8 @@ public:
 	virtual void Shutdown(EToolShutdownType ShutdownType) override;
 	virtual void OnTick(float DeltaTime) override;
 
+	// ISkeletalMeshGeometryIsolationAwareTool
+	virtual bool IsInputIsolationValidOnOutput() const override;
 	
 protected:
 
diff --git a/Engine/Plugins/Experimental/Animation/SkeletalMeshMorphTargetEditingTools/Source/SkeletalMeshMorphTargetEditingTools/Public/MorphTargetVertexSculptTool.h b/Engine/Plugins/Experimental/Animation/SkeletalMeshMorphTargetEditingTools/Source/SkeletalMeshMorphTargetEditingTools/Public/MorphTargetVertexSculptTool.h
--- a/Engine/Plugins/Experimental/Animation/SkeletalMeshMorphTargetEditingTools/Source/SkeletalMeshMorphTargetEditingTools/Public/MorphTargetVertexSculptTool.h
+++ b/Engine/Plugins/Experimental/Animation/SkeletalMeshMorphTargetEditingTools/Source/SkeletalMeshMorphTargetEditingTools/Public/MorphTargetVertexSculptTool.h
@@ -9,6 +9,7 @@
 #include "MeshDescription.h"
 #include "ReferenceSkeleton.h"
 #include "SkeletalMesh/SkeletalMeshToolsHelper.h"
+#include "SkeletalMesh/ISkeletalMeshGeometryIsolationAwareTool.h"
 
 #include "MorphTargetVertexSculptTool.generated.h"
 
@@ -34,7 +35,8 @@ UCLASS()
 class UMorphTargetVertexSculptTool : 
 	public UMeshVertexSculptTool,
 	public ISkeletalMeshEditingInterface,
-	public IMorphTargetEditingToolInterface
+	public IMorphTargetEditingToolInterface,
+	public ISkeletalMeshGeometryIsolationAwareTool 
 {
 	GENERATED_BODY()
 
@@ -54,6 +56,9 @@ public:
 	virtual void OnBeginStroke(const FRay& WorldRay) override;
 	virtual void OnEndStroke() override;
 
+	// ISkeletalMeshGeometryIsolationAwareTool
+	virtual bool IsInputIsolationValidOnOutput() const override;
+	
 protected:
 	friend class FMorphTargetVertexSculptNonSymmetricChange;
 	virtual void RegisterBrushes() override;
@@ -67,6 +72,7 @@ protected:
 	virtual void HandleSkeletalMeshModified(const TArray<FName>& Payload, const ESkeletalMeshNotifyType InNotifyType) override;
 	
 	
+	
 	void PoseSculptMesh(const TArray<FTransform>& ComponentSpaceTransforms, const TMap<FName, float>& MorphTargetWeights);
 
 	UPROPERTY()
diff --git a/Engine/Plugins/Runtime/MeshModelingToolset/Source/MeshModelingToolsEditorOnly/Private/SkeletalMesh/SkeletonEditingTool.cpp b/Engine/Plugins/Runtime/MeshModelingToolset/Source/MeshModelingToolsEditorOnly/Private/SkeletalMesh/SkeletonEditingTool.cpp
--- a/Engine/Plugins/Runtime/MeshModelingToolset/Source/MeshModelingToolsEditorOnly/Private/SkeletalMesh/SkeletonEditingTool.cpp
+++ b/Engine/Plugins/Runtime/MeshModelingToolset/Source/MeshModelingToolsEditorOnly/Private/SkeletalMesh/SkeletonEditingTool.cpp
@@ -1673,6 +1673,16 @@ bool USkeletonEditingTool::HasSelectedComponent() const
 	return !SelectionMechanic->GetActiveSelection().IsEmpty();
 }
 
+bool USkeletonEditingTool::IsInputIsolationValidOnOutput() const
+{
+	return true;
+}
+
+bool USkeletonEditingTool::IsInputSelectionValidOnOutput()
+{
+	return true;
+}
+
 void USkeletonEditingTool::SelectBone(const FName& InBoneName)
 {
 	TArray<FName> NewSelection = Selection;
diff --git a/Engine/Plugins/Runtime/MeshModelingToolset/Source/MeshModelingToolsEditorOnly/Private/SkeletalMesh/SkinWeightsBindingTool.cpp b/Engine/Plugins/Runtime/MeshModelingToolset/Source/MeshModelingToolsEditorOnly/Private/SkeletalMesh/SkinWeightsBindingTool.cpp
--- a/Engine/Plugins/Runtime/MeshModelingToolset/Source/MeshModelingToolsEditorOnly/Private/SkeletalMesh/SkinWeightsBindingTool.cpp
+++ b/Engine/Plugins/Runtime/MeshModelingToolset/Source/MeshModelingToolsEditorOnly/Private/SkeletalMesh/SkinWeightsBindingTool.cpp
@@ -309,6 +309,16 @@ TUniquePtr<UE::Geometry::FDynamicMeshOperator> USkinWeightsBindingTool::MakeNewO
 	return Op;
 }
 
+bool USkinWeightsBindingTool::IsInputIsolationValidOnOutput() const
+{
+	return true;
+}
+
+bool USkinWeightsBindingTool::IsInputSelectionValidOnOutput()
+{
+	return true;
+}
+
 bool USkinWeightsBindingTool::UpdateSkinWeightsFromDynamicMesh(FDynamicMesh3& InResultMesh) const
 {
 	using namespace UE::AnimationCore;
diff --git a/Engine/Plugins/Runtime/MeshModelingToolset/Source/MeshModelingToolsEditorOnly/Private/SkeletalMesh/SkinWeightsPaintTool.cpp b/Engine/Plugins/Runtime/MeshModelingToolset/Source/MeshModelingToolsEditorOnly/Private/SkeletalMesh/SkinWeightsPaintTool.cpp
--- a/Engine/Plugins/Runtime/MeshModelingToolset/Source/MeshModelingToolsEditorOnly/Private/SkeletalMesh/SkinWeightsPaintTool.cpp
+++ b/Engine/Plugins/Runtime/MeshModelingToolset/Source/MeshModelingToolsEditorOnly/Private/SkeletalMesh/SkinWeightsPaintTool.cpp
@@ -3108,6 +3108,16 @@ float USkinWeightsPaintTool::CalculateBrushFalloff(float Distance) const
 	return w;
 }
 
+bool USkinWeightsPaintTool::IsInputIsolationValidOnOutput() const
+{
+	return true;
+}
+
+bool USkinWeightsPaintTool::IsInputSelectionValidOnOutput()
+{
+	return true;
+}
+
 void USkinWeightsPaintTool::RegisterActions(FInteractiveToolActionSet& ActionSet)
 {
 	Super::RegisterActions(ActionSet);
diff --git a/Engine/Plugins/Runtime/MeshModelingToolset/Source/MeshModelingToolsEditorOnly/Public/SkeletalMesh/ISkeletalMeshGeometryIsolationAwareTool.h b/Engine/Plugins/Runtime/MeshModelingToolset/Source/MeshModelingToolsEditorOnly/Public/SkeletalMesh/ISkeletalMeshGeometryIsolationAwareTool.h
new file mode 100644
--- /dev/null
+++ b/Engine/Plugins/Runtime/MeshModelingToolset/Source/MeshModelingToolsEditorOnly/Public/SkeletalMesh/ISkeletalMeshGeometryIsolationAwareTool.h
@@ -0,0 +1,24 @@
+﻿// Copyright Epic Games, Inc. All Rights Reserved.
+
+#pragma once
+
+#include "UObject/Interface.h"
+
+#include "ISkeletalMeshGeometryIsolationAwareTool.generated.h"
+
+
+UINTERFACE(MinimalAPI)
+class USkeletalMeshGeometryIsolationAwareTool :
+	public UInterface
+{
+	GENERATED_BODY()
+};
+
+
+class ISkeletalMeshGeometryIsolationAwareTool : public IInterface
+{
+	GENERATED_BODY()
+
+public:
+	virtual bool IsInputIsolationValidOnOutput() const = 0;
+};
diff --git a/Engine/Plugins/Runtime/MeshModelingToolset/Source/MeshModelingToolsEditorOnly/Public/SkeletalMesh/SkeletonEditingTool.h b/Engine/Plugins/Runtime/MeshModelingToolset/Source/MeshModelingToolsEditorOnly/Public/SkeletalMesh/SkeletonEditingTool.h
--- a/Engine/Plugins/Runtime/MeshModelingToolset/Source/MeshModelingToolsEditorOnly/Public/SkeletalMesh/SkeletonEditingTool.h
+++ b/Engine/Plugins/Runtime/MeshModelingToolset/Source/MeshModelingToolsEditorOnly/Public/SkeletalMesh/SkeletonEditingTool.h
@@ -4,6 +4,7 @@
 
 #include "DynamicMeshBrushTool.h"
 #include "GroupTopology.h"
+#include "ISkeletalMeshGeometryIsolationAwareTool.h"
 #include "BaseTools/MeshSurfacePointMeshEditingTool.h"
 
 #include "SkeletalMeshNotifier.h"
@@ -135,7 +136,9 @@ UCLASS(MinimalAPI)
 class USkeletonEditingTool :
 	public USingleSelectionTool,
 	public IClickDragBehaviorTarget,
-	public ISkeletalMeshEditingInterface
+	public ISkeletalMeshEditingInterface,
+	public ISkeletalMeshGeometryIsolationAwareTool,
+	public IInteractiveToolManageGeometrySelectionAPI
 {
 	GENERATED_BODY()
 
@@ -196,6 +199,11 @@ public:
 
 	UE_API bool HasSelectedComponent() const;
 	
+	// ISkeletalMeshGeometryIsolationAwareTool
+	UE_API virtual bool IsInputIsolationValidOnOutput() const override;
+	
+	// IInteractiveToolManageGeometrySelectionAPI
+	UE_API virtual bool IsInputSelectionValidOnOutput() override;
 protected:
 
 	// UInteractiveTool overrides
diff --git a/Engine/Plugins/Runtime/MeshModelingToolset/Source/MeshModelingToolsEditorOnly/Public/SkeletalMesh/SkinWeightsBindingTool.h b/Engine/Plugins/Runtime/MeshModelingToolset/Source/MeshModelingToolsEditorOnly/Public/SkeletalMesh/SkinWeightsBindingTool.h
--- a/Engine/Plugins/Runtime/MeshModelingToolset/Source/MeshModelingToolsEditorOnly/Public/SkeletalMesh/SkinWeightsBindingTool.h
+++ b/Engine/Plugins/Runtime/MeshModelingToolset/Source/MeshModelingToolsEditorOnly/Public/SkeletalMesh/SkinWeightsBindingTool.h
@@ -6,6 +6,7 @@
 
 #include "BoneContainer.h"
 #include "InteractiveToolBuilder.h"
+#include "ISkeletalMeshGeometryIsolationAwareTool.h"
 #include "ModelingOperators.h"
 #include "SkeletalMeshEditingInterface.h"
 #include "BaseTools/MultiSelectionMeshEditingTool.h"
@@ -85,7 +86,9 @@ UCLASS(MinimalAPI)
 class USkinWeightsBindingTool :
 	public UMultiSelectionMeshEditingTool,
 	public UE::Geometry::IDynamicMeshOperatorFactory,
-	public ISkeletalMeshEditingInterface
+	public ISkeletalMeshEditingInterface,
+	public ISkeletalMeshGeometryIsolationAwareTool,
+    public IInteractiveToolManageGeometrySelectionAPI
 {
 	GENERATED_BODY()
 
@@ -108,6 +111,12 @@ public:
 	// IDynamicMeshOperatorFactory API
 	UE_API virtual TUniquePtr<UE::Geometry::FDynamicMeshOperator> MakeNewOperator() override;
 
+	// ISkeletalMeshGeometryIsolationAwareTool
+	UE_API virtual bool IsInputIsolationValidOnOutput() const override;
+	
+	// IInteractiveToolManageGeometrySelectionAPI
+	UE_API virtual bool IsInputSelectionValidOnOutput() override;
+	
 	UPROPERTY()
 	TObjectPtr<USkinWeightsBindingToolProperties> Properties = nullptr;
 	
diff --git a/Engine/Plugins/Runtime/MeshModelingToolset/Source/MeshModelingToolsEditorOnly/Public/SkeletalMesh/SkinWeightsPaintTool.h b/Engine/Plugins/Runtime/MeshModelingToolset/Source/MeshModelingToolsEditorOnly/Public/SkeletalMesh/SkinWeightsPaintTool.h
--- a/Engine/Plugins/Runtime/MeshModelingToolset/Source/MeshModelingToolsEditorOnly/Public/SkeletalMesh/SkinWeightsPaintTool.h
+++ b/Engine/Plugins/Runtime/MeshModelingToolset/Source/MeshModelingToolsEditorOnly/Public/SkeletalMesh/SkinWeightsPaintTool.h
@@ -9,6 +9,7 @@
 #include "BoneWeights.h"
 #include "DynamicSubmesh3.h"
 #include "GroupTopology.h"
+#include "ISkeletalMeshGeometryIsolationAwareTool.h"
 #include "SkeletalMeshAttributes.h"
 #include "Misc/Optional.h"
 #include "Containers/Map.h"
@@ -648,7 +649,11 @@ public:
 
 // An interactive tool for painting and editing skin weights.
 UCLASS(MinimalAPI)
-class USkinWeightsPaintTool : public UDynamicMeshBrushTool, public ISkeletalMeshEditingInterface
+class USkinWeightsPaintTool : 
+	public UDynamicMeshBrushTool, 
+	public ISkeletalMeshEditingInterface, 
+	public ISkeletalMeshGeometryIsolationAwareTool,
+	public IInteractiveToolManageGeometrySelectionAPI
 {
 	GENERATED_BODY()
 
@@ -841,6 +846,11 @@ public:
 	/** Returns the skeleton tree displayed when using this tool */
 	const TSharedPtr<SReferenceSkeletonTree>& GetToolSkeletonTree() const { return ToolSkeletonTree; }
 
+	// ISkeletalMeshGeometryIsolationAwareTool
+	UE_API virtual bool IsInputIsolationValidOnOutput() const override;
+	
+	// IInteractiveToolManageGeometrySelectionAPI
+	UE_API virtual bool IsInputSelectionValidOnOutput() override;
 protected:
 	// UInteractiveTool overrides
 	UE_API virtual void RegisterActions(FInteractiveToolActionSet& ActionSet) override;
