[SKMModelingTools] minor comment fix

#Synced-CL 52668932

diff --git a/Engine/Plugins/Animation/SkeletalMeshModelingTools/Source/SkeletalMeshModelingTools/Private/SkeletalMeshEditingCache.h b/Engine/Plugins/Animation/SkeletalMeshModelingTools/Source/SkeletalMeshModelingTools/Private/SkeletalMeshEditingCache.h
--- a/Engine/Plugins/Animation/SkeletalMeshModelingTools/Source/SkeletalMeshModelingTools/Private/SkeletalMeshEditingCache.h
+++ b/Engine/Plugins/Animation/SkeletalMeshModelingTools/Source/SkeletalMeshModelingTools/Private/SkeletalMeshEditingCache.h
@@ -252,7 +252,7 @@ protected:
 	void RebuildPreviewMesh();
 
 	// Pure computation: given an operation and current selection, computes the new isolation
-	// triangle set (in full mesh space) and remaps the selection to the new submesh space.
+	// triangle set (in full mesh space)
 	enum class EIsolationOperation
 	{
 		Isolate, 
